Geometry Cache Sequencer Support for 4.20

diff --git a/Engine/Plugins/Experimental/AlembicImporter/Source/AlembicLibrary/Private/AbcImporter.cpp b/Engine/Plugins/Experimental/AlembicImporter/Source/AlembicLibrary/Private/AbcImporter.cpp
--- a/Engine/Plugins/Experimental/AlembicImporter/Source/AlembicLibrary/Private/AbcImporter.cpp
+++ b/Engine/Plugins/Experimental/AlembicImporter/Source/AlembicLibrary/Private/AbcImporter.cpp
@@ -565,7 +565,9 @@ UGeometryCache* FAbcImporter::ImportAsGeometryCache(UObject* InParent, EObjectFl
 		{
 			Track->SetDuration(MaxDuration);
 		}
-
+		// Also store the number of frames in the cache
+		GeometryCache->SetFrameStartEnd(ImportSettings->SamplingSettings.FrameStart, ImportSettings->SamplingSettings.FrameEnd);
+		
 		// Update all geometry cache components, TODO move render-data from component to GeometryCache and allow for DDC population
 		for (TObjectIterator<UGeometryCacheComponent> CacheIt; CacheIt; ++CacheIt)
 		{
diff --git a/Engine/Plugins/Experimental/GeometryCache/GeometryCache.uplugin b/Engine/Plugins/Experimental/GeometryCache/GeometryCache.uplugin
--- a/Engine/Plugins/Experimental/GeometryCache/GeometryCache.uplugin
+++ b/Engine/Plugins/Experimental/GeometryCache/GeometryCache.uplugin
@@ -25,10 +25,24 @@
 				"Win64"
 			]
 		},
+		{
+			"Name" : "GeometryCacheSequencer",
+			"Type" : "Editor",
+			"LoadingPhase" : "Default",
+			"WhitelistPlatforms" :
+			[
+				"Win64"
+			]
+		},
 		{
 			"Name" : "GeometryCache",
 			"Type" : "Runtime",
 			"LoadingPhase" : "Default"
+		},
+		{
+			"Name" : "GeometryCacheTracks",
+			"Type" : "Runtime",
+			"LoadingPhase" : "Default"
 		}
 	]
 }
\ No newline at end of file
diff --git a/Engine/Plugins/Experimental/GeometryCache/Source/GeometryCache/Classes/GeometryCache.h b/Engine/Plugins/Experimental/GeometryCache/Source/GeometryCache/Classes/GeometryCache.h
--- a/Engine/Plugins/Experimental/GeometryCache/Source/GeometryCache/Classes/GeometryCache.h
+++ b/Engine/Plugins/Experimental/GeometryCache/Source/GeometryCache/Classes/GeometryCache.h
@@ -63,7 +63,23 @@ public:
 	UPROPERTY(VisibleAnywhere, Category=GeometryCache)
 	TArray<UGeometryCacheTrack*> Tracks;
 
+	/** Set the start and end frames for the GeometryCache */
+	void SetFrameStartEnd(int32 InStartFrame, int32 InEndFrame);
+
+	/** Get the start frame */
+	int32 GetStartFrame() const;
+
+	/** Get the end frame */
+	int32 GetEndFrame() const;
 private:
 	/** A fence which is used to keep track of the rendering thread releasing the geometry cache resources. */
 	FRenderCommandFence ReleaseResourcesFence;
+
+protected:
+	UPROPERTY(BlueprintReadOnly, Category = GeometryCache)
+	int32 StartFrame;
+
+	UPROPERTY(BlueprintReadOnly, Category = GeometryCache)
+	int32 EndFrame;
+
 };
diff --git a/Engine/Plugins/Experimental/GeometryCache/Source/GeometryCache/Classes/GeometryCacheComponent.h b/Engine/Plugins/Experimental/GeometryCache/Source/GeometryCache/Classes/GeometryCacheComponent.h
--- a/Engine/Plugins/Experimental/GeometryCache/Source/GeometryCache/Classes/GeometryCacheComponent.h
+++ b/Engine/Plugins/Experimental/GeometryCache/Source/GeometryCache/Classes/GeometryCacheComponent.h
@@ -79,6 +79,7 @@ class GEOMETRYCACHE_API UGeometryCacheComponent : public UMeshComponent
 	*/
 	void OnObjectReimported(UGeometryCache* ImportedGeometryCache);
 	
+public:
 	/** Start playback of GeometryCache */
 	UFUNCTION(BlueprintCallable, Category = "Components|GeometryCache")
 	void Play();
@@ -129,7 +130,7 @@ class GEOMETRYCACHE_API UGeometryCacheComponent : public UMeshComponent
 
 	/** Change the Geometry Cache used by this instance. */
 	UFUNCTION(BlueprintCallable, Category = "Components|GeometryCache")
-	virtual bool SetGeometryCache( UGeometryCache* NewGeomCache );
+	bool SetGeometryCache( UGeometryCache* NewGeomCache );
 	
 	/** Getter for Geometry cache instance referred by the component
 		Note: This getter is not exposed to blueprints as you can use the readonly Uproperty for that
@@ -146,16 +147,39 @@ class GEOMETRYCACHE_API UGeometryCacheComponent : public UMeshComponent
 
 	/** Set the current animation time for GeometryCache. Includes the influence of elapsed time and SetStartTimeOffset */
 	UFUNCTION(BlueprintCallable, Category = "Components|GeometryCache")
-	float GetAnimationTime();
+	float GetAnimationTime() const;
 
 	/** Set the current animation time for GeometryCache. Includes the influence of elapsed time and SetStartTimeOffset */
 	UFUNCTION(BlueprintCallable, Category = "Components|GeometryCache")
-	float GetPlaybackDirection();
+	float GetPlaybackDirection() const;
 
 	/** Geometry Cache instance referenced by the component */
 	UPROPERTY(EditAnywhere, BlueprintReadOnly, Category = GeometryCache)
 	UGeometryCache* GeometryCache;
 		
+	/** Get the duration of the playback */
+	UFUNCTION(BlueprintCallable, Category = "Components|GeometryCache")
+	float GetDuration() const;
+
+	/** Get the number of frames */
+	UFUNCTION(BlueprintCallable, Category = "Components|GeometryCache")
+	int32 GetNumberOfFrames() const;
+
+public:
+	/** Helper to get the frame of the ABC asset at this time*/
+	int32 GetFrameAtTime(const float Time) const;
+
+	/** Helper to get the time at this frame */
+	float GetTimeAtFrame(const int32 Frame) const;
+
+public:
+	/** Functions to override the default TickComponent */
+	void SetManualTick(bool bInManualTick);
+	bool GetManualTick() const;
+
+	UFUNCTION(BlueprintCallable, Category = "Components|GeometryCache")
+	void TickAtThisTime(const float Time, bool bInIsRunning, bool bInBackwards, bool bInIsLooping);
+
 protected:
 	/**
 	* Invalidate both the Matrix and Mesh sample indices
@@ -223,4 +247,7 @@ protected:
 	/** Duration of the animation (maximum time) */
 	UPROPERTY(BlueprintReadOnly, Category = GeometryCache)
 	float Duration;
+
+	UPROPERTY(EditAnywhere, Category = GeometryCache)
+	bool bManualTick;
 };
diff --git a/Engine/Plugins/Experimental/GeometryCache/Source/GeometryCache/Private/GeometryCache.cpp b/Engine/Plugins/Experimental/GeometryCache/Source/GeometryCache/Private/GeometryCache.cpp
--- a/Engine/Plugins/Experimental/GeometryCache/Source/GeometryCache/Private/GeometryCache.cpp
+++ b/Engine/Plugins/Experimental/GeometryCache/Source/GeometryCache/Private/GeometryCache.cpp
@@ -20,6 +20,8 @@ DEFINE_LOG_CATEGORY(LogGeometryCache);
 
 UGeometryCache::UGeometryCache(const FObjectInitializer& ObjectInitializer) : Super(ObjectInitializer)
 {
+	StartFrame = 0;
+	EndFrame = 0;
 }
 
 void UGeometryCache::PostInitProperties()
@@ -127,4 +129,20 @@ void UGeometryCache::AddTrack(UGeometryCacheTrack* Track)
 	Tracks.Add(Track);
 }
 
+void UGeometryCache::SetFrameStartEnd(int32 InStartFrame, int32 InEndFrame)
+{
+	StartFrame = InStartFrame;
+	EndFrame = InEndFrame;
+}
+
+int32 UGeometryCache::GetStartFrame() const
+{
+	return StartFrame;
+}
+
+int32 UGeometryCache::GetEndFrame() const
+{
+	return EndFrame;
+}
+
 #undef LOCTEXT_NAMESPACE // "GeometryCache"
\ No newline at end of file
diff --git a/Engine/Plugins/Experimental/GeometryCache/Source/GeometryCache/Private/GeometryCacheComponent.cpp b/Engine/Plugins/Experimental/GeometryCache/Source/GeometryCache/Private/GeometryCacheComponent.cpp
--- a/Engine/Plugins/Experimental/GeometryCache/Source/GeometryCache/Private/GeometryCacheComponent.cpp
+++ b/Engine/Plugins/Experimental/GeometryCache/Source/GeometryCache/Private/GeometryCacheComponent.cpp
@@ -27,6 +27,7 @@ UGeometryCacheComponent::UGeometryCacheComponent(const FObjectInitializer& Objec
 	StartTimeOffset = 0.0f;
 	PlaybackSpeed = 1.0f;
 	Duration = 0.0f;
+	bManualTick = false;
 }
 
 void UGeometryCacheComponent::BeginDestroy()
@@ -90,7 +91,7 @@ void UGeometryCacheComponent::OnUnregister()
 void UGeometryCacheComponent::TickComponent(float DeltaTime, enum ELevelTick TickType, FActorComponentTickFunction *ThisTickFunction)
 {
 	SCOPE_CYCLE_COUNTER(STAT_GeometryCacheComponent_TickComponent);
-	if (GeometryCache && bRunning)
+	if (GeometryCache && bRunning && !bManualTick)
 	{
 		// Increase total elapsed time since BeginPlay according to PlayDirection and speed
 		ElapsedTime += (DeltaTime * PlayDirection * GetPlaybackSpeed());
@@ -138,6 +139,57 @@ void UGeometryCacheComponent::TickComponent(float DeltaTime, enum ELevelTick Tic
 	}
 }
 
+void UGeometryCacheComponent::SetManualTick(bool bInManualTick)
+{
+	bManualTick = bInManualTick;
+}
+bool UGeometryCacheComponent::GetManualTick() const
+{
+	return bManualTick;
+}
+
+void UGeometryCacheComponent::TickAtThisTime(const float Time, bool bInIsRunning, bool bInBackwards, bool bInIsLooping)
+{
+	if (bManualTick && GeometryCache && bRunning)
+	{
+		// Game thread update:
+		// This mainly just updates the matrix and bounding boxes. All render state (meshes) is done on the render thread
+		bool bUpdatedBoundsOrMatrix = false;
+		for (int32 TrackIndex = 0; TrackIndex < NumTracks; ++TrackIndex)
+		{
+			bUpdatedBoundsOrMatrix |= UpdateTrackSection(TrackIndex);
+		}
+
+		if (bUpdatedBoundsOrMatrix)
+		{
+			UpdateLocalBounds();
+			// Mark the transform as dirty, so the bounds are updated and sent to the render thread
+			MarkRenderTransformDirty();
+		}
+
+		// The actual current playback speed. The PlaybackSpeed variable contains the speed it would
+		// play back at if it were running regardless of if we're running or not. The renderer
+		// needs the actual playback speed if not a paused animation with explicit motion vectors
+		// would just keep on blurring as if it were moving even when paused.
+		float ActualPlaybackSpeed = (bInIsRunning) ? PlaybackSpeed : 0.0f;
+
+		// Schedule an update on the render thread
+		if (FGeometryCacheSceneProxy* CastedProxy = static_cast<FGeometryCacheSceneProxy*>(SceneProxy))
+		{
+			ENQUEUE_UNIQUE_RENDER_COMMAND_FIVEPARAMETER(
+				FGeometryCacheUpdateAnimation,
+				FGeometryCacheSceneProxy*, SceneProxy, CastedProxy,
+				float, AnimationTime, Time,
+				bool, bLooping, bInIsLooping,
+				bool, bIsPlayingBackwards, bInBackwards,
+				float, PlaybackSpeed, ActualPlaybackSpeed,
+				{
+					SceneProxy->UpdateAnimation(AnimationTime, bLooping, bIsPlayingBackwards,PlaybackSpeed);
+				});
+		}
+	}
+}
+
 FBoxSphereBounds UGeometryCacheComponent::CalcBounds(const FTransform& LocalToWorld) const
 {
 	return LocalBounds.TransformBy(LocalToWorld);
@@ -365,17 +417,22 @@ void UGeometryCacheComponent::SetStartTimeOffset(const float NewStartTimeOffset)
 	MarkRenderStateDirty();
 }
 
-float UGeometryCacheComponent::GetAnimationTime()
+float UGeometryCacheComponent::GetAnimationTime() const
 {
 	const float ClampedStartTimeOffset = FMath::Clamp(StartTimeOffset, -14400.0f, 14400.0f);
 	return ElapsedTime + ClampedStartTimeOffset;
 }
 
-float UGeometryCacheComponent::GetPlaybackDirection()
+float UGeometryCacheComponent::GetPlaybackDirection() const
 {
 	return PlayDirection;
 }
 
+float UGeometryCacheComponent::GetDuration() const
+{
+	return Duration;
+}
+
 void UGeometryCacheComponent::PlayReversedFromEnd()
 {
 	ElapsedTime = Duration;
@@ -408,6 +465,31 @@ void UGeometryCacheComponent::ReleaseResources()
 	DetachFence.BeginFence();
 }
 
+int32 UGeometryCacheComponent::GetFrameAtTime(const float Time) const
+{
+
+	const float FrameTime = GetNumberOfFrames() > 1 ? Duration / (float)(GetNumberOfFrames() - 1) : 0.0f;
+	const int32 NormalizedFrame =  FMath::Clamp(FMath::RoundToInt(Time / FrameTime), 0, GetNumberOfFrames() - 1);
+	const int32 StartFrame = GeometryCache != nullptr ? GeometryCache->GetStartFrame() : 0;
+	return StartFrame + NormalizedFrame; //
+}
+
+float UGeometryCacheComponent::GetTimeAtFrame(const int32 Frame) const
+{
+	const float FrameTime = GetNumberOfFrames() > 1 ? Duration / (float)(GetNumberOfFrames() - 1) : 0.0f;
+	const int32 StartFrame = GeometryCache != nullptr ? GeometryCache->GetStartFrame() : 0;
+
+	return FMath::Clamp(FrameTime * (Frame - StartFrame), 0.0f, Duration);
+}
+
+int32 UGeometryCacheComponent::GetNumberOfFrames() const
+{
+	if (GeometryCache)
+	{
+		return GeometryCache->GetEndFrame() - GeometryCache->GetStartFrame()  + 1;
+	}
+	return 0;
+}
 #if WITH_EDITOR
 void UGeometryCacheComponent::PreEditUndo()
 {
diff --git a/Engine/Plugins/Experimental/GeometryCache/Source/GeometryCacheSequencer/Classes/GeometryCacheTrackEditor.cpp b/Engine/Plugins/Experimental/GeometryCache/Source/GeometryCacheSequencer/Classes/GeometryCacheTrackEditor.cpp
new file mode 100644
--- /dev/null
+++ b/Engine/Plugins/Experimental/GeometryCache/Source/GeometryCacheSequencer/Classes/GeometryCacheTrackEditor.cpp
@@ -0,0 +1,376 @@
+// Copyright 1998-2018 Epic Games, Inc. All Rights Reserved.
+
+#include "GeometryCacheTrackEditor.h"
+#include "Rendering/DrawElements.h"
+#include "Widgets/SBoxPanel.h"
+#include "Framework/MultiBox/MultiBoxBuilder.h"
+#include "GameFramework/Actor.h"
+#include "Modules/ModuleManager.h"
+#include "Layout/WidgetPath.h"
+#include "Framework/Application/MenuStack.h"
+#include "Framework/Application/SlateApplication.h"
+#include "Widgets/Layout/SBox.h"
+#include "SequencerSectionPainter.h"
+#include "Editor/UnrealEdEngine.h"
+#include "UnrealEdGlobals.h"
+#include "MovieSceneGeometryCacheTrack.h"
+#include "MovieSceneGeometryCacheSection.h"
+#include "CommonMovieSceneTools.h"
+#include "ContentBrowserModule.h"
+#include "SequencerUtilities.h"
+#include "ISectionLayoutBuilder.h"
+#include "EditorStyleSet.h"
+#include "MovieSceneTimeHelpers.h"
+#include "Fonts/FontMeasure.h"
+#include "SequencerTimeSliderController.h"
+#include "Misc/MessageDialog.h"
+#include "Framework/Notifications/NotificationManager.h"
+#include "Widgets/Notifications/SNotificationList.h"
+#include "GeometryCacheComponent.h"
+#include "GeometryCache.h"
+
+namespace GeometryCacheEditorConstants
+{
+	// @todo Sequencer Allow this to be customizable
+	const uint32 AnimationTrackHeight = 20;
+}
+
+#define LOCTEXT_NAMESPACE "FGeometryCacheTrackEditor"
+
+static UGeometryCacheComponent* AcquireGeometryCacheFromObjectGuid(const FGuid& Guid, TSharedPtr<ISequencer> SequencerPtr)
+{
+	UObject* BoundObject = SequencerPtr.IsValid() ? SequencerPtr->FindSpawnedObjectOrTemplate(Guid) : nullptr;
+
+	if (AActor* Actor = Cast<AActor>(BoundObject))
+	{
+		TInlineComponentArray<UGeometryCacheComponent*> GeometryMeshComponents;
+		Actor->GetComponents(GeometryMeshComponents);
+
+		for (int32 j = 0; j <GeometryMeshComponents.Num(); ++j)
+		{
+			UGeometryCacheComponent* GeometryMeshComp = GeometryMeshComponents[j];
+			if (GeometryMeshComp)
+			{
+				return GeometryMeshComp;
+			}
+		}
+	}
+	else if(UGeometryCacheComponent* GeometryMeshComp = Cast<UGeometryCacheComponent>(BoundObject))
+	{
+		if (GeometryMeshComp->GetGeometryCache())
+		{
+			return GeometryMeshComp;
+		}
+	}
+
+	return nullptr;
+}
+
+
+FGeometryCacheSection::FGeometryCacheSection( UMovieSceneSection& InSection, TWeakPtr<ISequencer> InSequencer)
+	: Section(*CastChecked<UMovieSceneGeometryCacheSection>(&InSection))
+	, Sequencer(InSequencer)
+	, InitialStartOffsetDuringResize(0.f)
+	, InitialStartTimeDuringResize(0)
+{ }
+
+
+UMovieSceneSection* FGeometryCacheSection::GetSectionObject()
+{ 
+	return &Section;
+}
+
+
+FText FGeometryCacheSection::GetSectionTitle() const
+{
+	if (Section.Params.GeometryCache != nullptr  && Section.Params.GeometryCache->GetOwner() != nullptr)
+	{
+		return FText::FromString(Section.Params.GeometryCache->GetOwner()->GetName());
+	}
+	return LOCTEXT("NoGeometryCacheSection", "No GeometryCache");
+}
+
+
+float FGeometryCacheSection::GetSectionHeight() const
+{
+	return (float)GeometryCacheEditorConstants::AnimationTrackHeight;
+}
+
+
+int32 FGeometryCacheSection::OnPaintSection( FSequencerSectionPainter& Painter ) const
+{
+	const ESlateDrawEffect DrawEffects = Painter.bParentEnabled ? ESlateDrawEffect::None : ESlateDrawEffect::DisabledEffect;
+	
+	const FTimeToPixel& TimeToPixelConverter = Painter.GetTimeConverter();
+
+	int32 LayerId = Painter.PaintSectionBackground();
+
+	static const FSlateBrush* GenericDivider = FEditorStyle::GetBrush("Sequencer.GenericDivider");
+
+	if (!Section.HasStartFrame() || !Section.HasEndFrame())
+	{
+		return LayerId;
+	}
+
+	// Add lines where the animation starts and ends/loops
+	float AnimPlayRate = FMath::IsNearlyZero(Section.Params.PlayRate) ? 1.0f : Section.Params.PlayRate;
+	float SeqLength = (Section.Params.GetSequenceLength() - (Section.Params.StartOffset + Section.Params.EndOffset)) / AnimPlayRate;
+
+	FFrameRate TickResolution = TimeToPixelConverter.GetTickResolution();
+	if (!FMath::IsNearlyZero(SeqLength, KINDA_SMALL_NUMBER) && SeqLength > 0)
+	{
+		float MaxOffset  = Section.GetRange().Size<FFrameTime>() / TickResolution;
+		float OffsetTime = SeqLength;
+		float StartTime  = Section.GetInclusiveStartFrame() / TickResolution;
+
+		while (OffsetTime < MaxOffset)
+		{
+			float OffsetPixel = TimeToPixelConverter.SecondsToPixel(StartTime + OffsetTime) - TimeToPixelConverter.SecondsToPixel(StartTime);
+
+			FSlateDrawElement::MakeBox(
+				Painter.DrawElements,
+				LayerId,
+				Painter.SectionGeometry.MakeChild(
+					FVector2D(2.f, Painter.SectionGeometry.Size.Y-2.f),
+					FSlateLayoutTransform(FVector2D(OffsetPixel, 1.f))
+				).ToPaintGeometry(),
+				GenericDivider,
+				DrawEffects
+			);
+
+			OffsetTime += SeqLength;
+		}
+	}
+
+	TSharedPtr<ISequencer> SequencerPtr = Sequencer.Pin();
+	if (Painter.bIsSelected && SequencerPtr.IsValid())
+	{
+		FFrameTime CurrentTime = SequencerPtr->GetLocalTime().Time;
+		if (Section.GetRange().Contains(CurrentTime.FrameNumber) && Section.Params.GeometryCache != nullptr)
+		{
+			const float Time = TimeToPixelConverter.FrameToPixel(CurrentTime); 
+
+			// Draw the current time next to the scrub handle
+			const float AnimTime = Section.MapTimeToAnimation(CurrentTime, TickResolution);			int32 FrameTime = Section.Params.GeometryCache->GetFrameAtTime(AnimTime);
+			FString FrameString = FString::FromInt(FrameTime);
+
+			const FSlateFontInfo SmallLayoutFont = FCoreStyle::GetDefaultFontStyle("Bold", 10);
+			const TSharedRef< FSlateFontMeasure > FontMeasureService = FSlateApplication::Get().GetRenderer()->GetFontMeasureService();
+			FVector2D TextSize = FontMeasureService->Measure(FrameString, SmallLayoutFont);
+
+			// Flip the text position if getting near the end of the view range
+			static const float TextOffsetPx = 10.f;
+			bool  bDrawLeft = (Painter.SectionGeometry.Size.X - Time) < (TextSize.X + 22.f) - TextOffsetPx;
+			float TextPosition = bDrawLeft ? Time - TextSize.X - TextOffsetPx : Time + TextOffsetPx;
+			//handle mirrored labels
+			const float MajorTickHeight = 9.0f; 
+			FVector2D TextOffset(TextPosition, Painter.SectionGeometry.Size.Y - (MajorTickHeight + TextSize.Y));
+
+			const FLinearColor DrawColor = FEditorStyle::GetSlateColor("SelectionColor").GetColor(FWidgetStyle());
+			const FVector2D BoxPadding = FVector2D(4.0f, 2.0f);
+			// draw time string
+	
+			FSlateDrawElement::MakeBox(
+				Painter.DrawElements,
+				LayerId + 5,
+				Painter.SectionGeometry.ToPaintGeometry(TextOffset - BoxPadding, TextSize + 2.0f * BoxPadding),
+				FEditorStyle::GetBrush("WhiteBrush"),
+				ESlateDrawEffect::None,
+				FLinearColor::Black.CopyWithNewOpacity(0.5f)
+			);
+
+			FSlateDrawElement::MakeText(
+				Painter.DrawElements,
+				LayerId + 6,
+				Painter.SectionGeometry.ToPaintGeometry(TextOffset, TextSize),
+				FrameString,
+				SmallLayoutFont,
+				DrawEffects,
+				DrawColor
+			);
+
+		}
+	}
+	
+	return LayerId;
+}
+
+void FGeometryCacheSection::BeginResizeSection()
+{
+	InitialStartOffsetDuringResize = Section.Params.StartOffset;
+	InitialStartTimeDuringResize   = Section.HasStartFrame() ? Section.GetInclusiveStartFrame() : 0;
+}
+
+void FGeometryCacheSection::ResizeSection(ESequencerSectionResizeMode ResizeMode, FFrameNumber ResizeTime)
+{
+	// Adjust the start offset when resizing from the beginning
+	if (ResizeMode == SSRM_LeadingEdge)
+	{
+		FFrameRate FrameRate   = Section.GetTypedOuter<UMovieScene>()->GetTickResolution();
+		float      StartOffset = (ResizeTime - InitialStartTimeDuringResize) / FrameRate * Section.Params.PlayRate;
+
+		StartOffset += InitialStartOffsetDuringResize;
+
+		// Ensure start offset is not less than 0 and adjust ResizeTime
+		if (StartOffset < 0)
+		{
+			ResizeTime = ResizeTime - ((StartOffset * Section.Params.PlayRate) * FrameRate).RoundToFrame();
+
+			StartOffset = 0.f;
+		}
+
+		Section.Params.StartOffset = StartOffset;
+	}
+
+	ISequencerSection::ResizeSection(ResizeMode, ResizeTime);
+}
+
+void FGeometryCacheSection::BeginSlipSection()
+{
+	BeginResizeSection();
+}
+
+void FGeometryCacheSection::SlipSection(double SlipTime)
+{
+	float StartOffset = (SlipTime - InitialStartTimeDuringResize / Section.GetTypedOuter<UMovieScene>()->GetTickResolution()) * Section.Params.PlayRate;
+	StartOffset += InitialStartOffsetDuringResize;
+
+	// Ensure start offset is not less than 0
+	if (StartOffset < 0)
+	{
+		StartOffset = 0.f;
+	}
+
+	Section.Params.StartOffset = StartOffset;
+
+	ISequencerSection::SlipSection(SlipTime);
+}
+
+
+FGeometryCacheTrackEditor::FGeometryCacheTrackEditor( TSharedRef<ISequencer> InSequencer )
+	: FMovieSceneTrackEditor( InSequencer ) 
+{ }
+
+
+TSharedRef<ISequencerTrackEditor> FGeometryCacheTrackEditor::CreateTrackEditor( TSharedRef<ISequencer> InSequencer )
+{
+	return MakeShareable( new FGeometryCacheTrackEditor( InSequencer ) );
+}
+
+
+bool FGeometryCacheTrackEditor::SupportsType( TSubclassOf<UMovieSceneTrack> Type ) const
+{
+	return Type == UMovieSceneGeometryCacheTrack::StaticClass();
+}
+
+
+TSharedRef<ISequencerSection> FGeometryCacheTrackEditor::MakeSectionInterface( UMovieSceneSection& SectionObject, UMovieSceneTrack& Track, FGuid ObjectBinding )
+{
+	check( SupportsType( SectionObject.GetOuter()->GetClass() ) );
+	
+	return MakeShareable( new FGeometryCacheSection(SectionObject, GetSequencer()) );
+}
+
+void FGeometryCacheTrackEditor::BuildObjectBindingTrackMenu(FMenuBuilder& MenuBuilder, const FGuid& ObjectBinding, const UClass* ObjectClass)
+{
+	if (ObjectClass->IsChildOf(UGeometryCacheComponent::StaticClass()) || ObjectClass->IsChildOf(AActor::StaticClass()))
+	{
+		UGeometryCacheComponent* GeomMeshComp = AcquireGeometryCacheFromObjectGuid(ObjectBinding, GetSequencer());
+
+		if (GeomMeshComp)
+		{
+			UMovieSceneTrack* Track = nullptr;
+
+			MenuBuilder.AddMenuEntry(
+				NSLOCTEXT("Sequencer", "AddGeometryCache", "Geometry Cache"),
+				NSLOCTEXT("Sequencer", "AddGeometryCacheTooltip", "Adds a Geometry Cache track."),
+				FSlateIcon(),
+				FUIAction(
+					FExecuteAction::CreateSP(this, &FGeometryCacheTrackEditor::BuildGeometryCacheTrack, ObjectBinding, GeomMeshComp, Track)
+				)
+			);
+		}
+	}
+}
+
+void FGeometryCacheTrackEditor::BuildGeometryCacheTrack(FGuid ObjectBinding, UGeometryCacheComponent *GeomCacheComp, UMovieSceneTrack* Track)
+{
+	TSharedPtr<ISequencer> SequencerPtr = GetSequencer();
+
+	if (SequencerPtr.IsValid() && ObjectBinding.IsValid())
+	{
+		UObject* Object = SequencerPtr->FindSpawnedObjectOrTemplate(ObjectBinding);
+		if (Object)
+		{
+			const FScopedTransaction Transaction(LOCTEXT("AddGeometryCache_Transaction", "Add Geometry Cache"));
+			AnimatablePropertyChanged(FOnKeyProperty::CreateRaw(this, &FGeometryCacheTrackEditor::AddKeyInternal, Object, GeomCacheComp, Track));
+		}
+	}
+}
+
+FKeyPropertyResult FGeometryCacheTrackEditor::AddKeyInternal( FFrameNumber KeyTime, UObject* Object, UGeometryCacheComponent* GeomCacheComp, UMovieSceneTrack* Track)
+{
+	FKeyPropertyResult KeyPropertyResult;
+
+	FFindOrCreateHandleResult HandleResult = FindOrCreateHandleToObject( Object );
+	FGuid ObjectHandle = HandleResult.Handle;
+	KeyPropertyResult.bHandleCreated |= HandleResult.bWasCreated;
+	if (ObjectHandle.IsValid())
+	{
+		if (!Track)
+		{
+			Track = AddTrack(GetSequencer()->GetFocusedMovieSceneSequence()->GetMovieScene(), ObjectHandle, UMovieSceneGeometryCacheTrack::StaticClass(), NAME_None);
+			KeyPropertyResult.bTrackCreated = true;
+		}
+
+		if (ensure(Track))
+		{
+			Track->Modify();
+
+			UMovieSceneSection* NewSection = Cast<UMovieSceneGeometryCacheTrack>(Track)->AddNewAnimation( KeyTime, GeomCacheComp );
+			KeyPropertyResult.bTrackModified = true;
+
+			GetSequencer()->EmptySelection();
+			GetSequencer()->SelectSection(NewSection);
+			GetSequencer()->ThrobSectionSelection();
+		}
+	}
+
+	return KeyPropertyResult;
+}
+
+
+TSharedPtr<SWidget> FGeometryCacheTrackEditor::BuildOutlinerEditWidget(const FGuid& ObjectBinding, UMovieSceneTrack* Track, const FBuildEditWidgetParams& Params)
+{
+	UGeometryCacheComponent* GeomMeshComp = AcquireGeometryCacheFromObjectGuid(ObjectBinding, GetSequencer());
+
+	if (GeomMeshComp)
+	{
+		TWeakPtr<ISequencer> WeakSequencer = GetSequencer();
+
+		auto SubMenuCallback = [=]() -> TSharedRef<SWidget>
+		{
+			FMenuBuilder MenuBuilder(true, nullptr);
+
+			BuildGeometryCacheTrack(ObjectBinding, GeomMeshComp, Track);
+			
+			return MenuBuilder.MakeWidget();
+		};
+
+		return SNew(SHorizontalBox)
+			+ SHorizontalBox::Slot()
+			.AutoWidth()
+			.VAlign(VAlign_Center)
+			[
+				FSequencerUtilities::MakeAddButton(LOCTEXT("GeomCacheText", "Geometry Cache"), FOnGetContent::CreateLambda(SubMenuCallback), Params.NodeIsHovered)
+			];
+	}
+	else
+	{
+		return TSharedPtr<SWidget>();
+	}
+	
+}
+
+
+#undef LOCTEXT_NAMESPACE
diff --git a/Engine/Plugins/Experimental/GeometryCache/Source/GeometryCacheSequencer/Classes/GeometryCacheTrackEditor.h b/Engine/Plugins/Experimental/GeometryCache/Source/GeometryCacheSequencer/Classes/GeometryCacheTrackEditor.h
new file mode 100644
--- /dev/null
+++ b/Engine/Plugins/Experimental/GeometryCache/Source/GeometryCacheSequencer/Classes/GeometryCacheTrackEditor.h
@@ -0,0 +1,102 @@
+// Copyright 1998-2018 Epic Games, Inc. All Rights Reserved.
+
+#pragma once
+
+#include "CoreMinimal.h"
+#include "Misc/Guid.h"
+#include "Templates/SubclassOf.h"
+#include "Widgets/SWidget.h"
+#include "ISequencer.h"
+#include "MovieSceneTrack.h"
+#include "ISequencerSection.h"
+#include "ISequencerTrackEditor.h"
+#include "MovieSceneTrackEditor.h"
+
+struct FAssetData;
+class FMenuBuilder;
+class FSequencerSectionPainter;
+class UMovieSceneGeometryCacheSection;
+class UGeometryCacheComponent;
+
+/**
+ * Tools for Geometry Cache tracks
+ */
+class FGeometryCacheTrackEditor : public FMovieSceneTrackEditor
+{
+public:
+	/**
+	 * Constructor
+	 *
+	 * @param InSequencer The sequencer instance to be used by this tool
+	 */
+	FGeometryCacheTrackEditor( TSharedRef<ISequencer> InSequencer );
+
+	/** Virtual destructor. */
+	virtual ~FGeometryCacheTrackEditor() { }
+
+	/**
+	 * Creates an instance of this class.  Called by a sequencer 
+	 *
+	 * @param OwningSequencer The sequencer instance to be used by this tool
+	 * @return The new instance of this class
+	 */
+	static TSharedRef<ISequencerTrackEditor> CreateTrackEditor( TSharedRef<ISequencer> OwningSequencer );
+
+public:
+
+	// ISequencerTrackEditor interface
+	virtual void BuildObjectBindingTrackMenu(FMenuBuilder& MenuBuilder, const FGuid& ObjectBinding, const UClass* ObjectClass) override;
+	virtual TSharedRef<ISequencerSection> MakeSectionInterface( UMovieSceneSection& SectionObject, UMovieSceneTrack& Track, FGuid ObjectBinding ) override;
+	virtual bool SupportsType( TSubclassOf<UMovieSceneTrack> Type ) const override;
+	virtual TSharedPtr<SWidget> BuildOutlinerEditWidget(const FGuid& ObjectBinding, UMovieSceneTrack* Track, const FBuildEditWidgetParams& Params) override;
+
+private:
+
+	void BuildGeometryCacheTrack(FGuid ObjectBinding, UGeometryCacheComponent *GeomCacheComp, UMovieSceneTrack* Track);
+
+	/** Delegate for AnimatablePropertyChanged in AddKey */
+	FKeyPropertyResult AddKeyInternal(FFrameNumber KeyTime, UObject* Object, UGeometryCacheComponent* GeomCacheComp, UMovieSceneTrack* Track);
+};
+
+
+/** Class for animation sections */
+class FGeometryCacheSection
+	: public ISequencerSection
+	, public TSharedFromThis<FGeometryCacheSection>
+{
+public:
+
+	/** Constructor. */
+	FGeometryCacheSection( UMovieSceneSection& InSection, TWeakPtr<ISequencer> InSequencer);
+
+	/** Virtual destructor. */
+	virtual ~FGeometryCacheSection() { }
+
+public:
+
+	// ISequencerSection interface
+
+	virtual UMovieSceneSection* GetSectionObject() override;
+	virtual FText GetSectionTitle() const override;
+	virtual float GetSectionHeight() const override;
+	virtual int32 OnPaintSection( FSequencerSectionPainter& Painter ) const override;
+	virtual void BeginResizeSection() override;
+	virtual void ResizeSection(ESequencerSectionResizeMode ResizeMode, FFrameNumber ResizeTime) override;
+	virtual void BeginSlipSection() override;
+	virtual void SlipSection(double SlipTime) override;
+
+
+private:
+
+	/** The section we are visualizing */
+	UMovieSceneGeometryCacheSection& Section;
+
+	/** Used to draw animation frame, need selection state and local time*/
+	TWeakPtr<ISequencer> Sequencer;
+
+	/** Cached start offset value valid only during resize */
+	float InitialStartOffsetDuringResize;
+	
+	/** Cached start time valid only during resize */
+	FFrameNumber InitialStartTimeDuringResize;
+};
diff --git a/Engine/Plugins/Experimental/GeometryCache/Source/GeometryCacheSequencer/GeometryCacheSequencer.Build.cs b/Engine/Plugins/Experimental/GeometryCache/Source/GeometryCacheSequencer/GeometryCacheSequencer.Build.cs
new file mode 100644
--- /dev/null
+++ b/Engine/Plugins/Experimental/GeometryCache/Source/GeometryCacheSequencer/GeometryCacheSequencer.Build.cs
@@ -0,0 +1,34 @@
+// Copyright 1998-2018 Epic Games, Inc. All Rights Reserved.
+
+namespace UnrealBuildTool.Rules
+{
+	public class GeometryCacheSequencer : ModuleRules
+	{
+		public GeometryCacheSequencer(ReadOnlyTargetRules Target) : base(Target)
+		{
+            OptimizeCode = CodeOptimization.Never;
+
+            PrivateDependencyModuleNames.AddRange(
+				new string[]
+				{
+					"AssetTools",
+					"Core",
+					"CoreUObject",
+					"EditorStyle",
+					"Engine",
+					"MovieScene",
+					"MovieSceneTools",
+					"MovieSceneTracks",
+					"RHI",
+					"Sequencer",
+					"Slate",
+					"SlateCore",
+                    "TimeManagement",
+					"UnrealEd",
+                    "GeometryCacheTracks",
+                    "GeometryCache"
+                }
+			);
+        }
+	}
+}
diff --git a/Engine/Plugins/Experimental/GeometryCache/Source/GeometryCacheSequencer/Private/GeometryCacheSequencerModule.cpp b/Engine/Plugins/Experimental/GeometryCache/Source/GeometryCacheSequencer/Private/GeometryCacheSequencerModule.cpp
new file mode 100644
--- /dev/null
+++ b/Engine/Plugins/Experimental/GeometryCache/Source/GeometryCacheSequencer/Private/GeometryCacheSequencerModule.cpp
@@ -0,0 +1,13 @@
+// Copyright 1998-2018 Epic Games, Inc. All Rights Reserved.
+#include "GeometryCacheSequencerModule.h"
+#include "ISequencerModule.h"
+#include "Modules/ModuleInterface.h"
+#include "Modules/ModuleManager.h"
+
+#define LOCTEXT_NAMESPACE "GeometryCacheSequencer"
+
+
+
+IMPLEMENT_MODULE(FGeometryCacheSequencerModule, GeometryCacheSequencer);
+
+#undef LOCTEXT_NAMESPACE
diff --git a/Engine/Plugins/Experimental/GeometryCache/Source/GeometryCacheSequencer/Public/GeometryCacheSequencerModule.h b/Engine/Plugins/Experimental/GeometryCache/Source/GeometryCacheSequencer/Public/GeometryCacheSequencerModule.h
new file mode 100644
--- /dev/null
+++ b/Engine/Plugins/Experimental/GeometryCache/Source/GeometryCacheSequencer/Public/GeometryCacheSequencerModule.h
@@ -0,0 +1,38 @@
+// Copyright 1998-2018 Epic Games, Inc. All Rights Reserved.
+
+#pragma once
+#include "ISequencerModule.h"
+#include "CoreMinimal.h"
+#include "Modules/ModuleInterface.h"
+#include "Modules/ModuleManager.h"
+#include "GeometryCacheTrackEditor.h"
+
+/**
+ * The public interface to this module
+ */
+class FGeometryCacheSequencerModule : public IModuleInterface
+{
+public:
+
+	virtual void StartupModule() override
+	{
+
+		ISequencerModule& SequencerModule = FModuleManager::Get().LoadModuleChecked<ISequencerModule>("Sequencer");
+		TrackEditorBindingHandle = SequencerModule.RegisterTrackEditor(FOnCreateTrackEditor::CreateStatic(&FGeometryCacheTrackEditor::CreateTrackEditor));
+
+	}
+
+	virtual void ShutdownModule() override
+	{
+
+		ISequencerModule* SequencerModulePtr = FModuleManager::Get().GetModulePtr<ISequencerModule>("Sequencer");
+		if (SequencerModulePtr)
+		{
+			SequencerModulePtr->UnRegisterTrackEditor(TrackEditorBindingHandle);
+		}
+
+	}
+
+	FDelegateHandle TrackEditorBindingHandle;
+};
+
diff --git a/Engine/Plugins/Experimental/GeometryCache/Source/GeometryCacheTracks/Classes/MovieSceneGeometryCacheSection.h b/Engine/Plugins/Experimental/GeometryCache/Source/GeometryCacheTracks/Classes/MovieSceneGeometryCacheSection.h
new file mode 100644
--- /dev/null
+++ b/Engine/Plugins/Experimental/GeometryCache/Source/GeometryCacheTracks/Classes/MovieSceneGeometryCacheSection.h
@@ -0,0 +1,89 @@
+// Copyright 1998-2018 Epic Games, Inc. All Rights Reserved.
+
+#pragma once
+
+#include "CoreMinimal.h"
+#include "UObject/ObjectMacros.h"
+#include "MovieSceneSection.h"
+#include "GeometryCacheComponent.h"
+#include "Channels/MovieSceneFloatChannel.h"
+#include "MovieSceneGeometryCacheSection.generated.h"
+
+USTRUCT()
+struct FMovieSceneGeometryCacheParams
+{
+	GENERATED_BODY()
+
+	FMovieSceneGeometryCacheParams();
+
+	/** Gets the animation duration, modified by play rate */
+	float GetDuration() const { return FMath::IsNearlyZero(PlayRate) || GeometryCache == nullptr ? 0.f : GeometryCache->GetDuration() / PlayRate; }
+
+	/** Gets the animation sequence length, not modified by play rate */
+	float GetSequenceLength() const { return GeometryCache != nullptr ? GeometryCache->GetDuration() : 0.f; }
+
+	/** The animation this section plays */
+	UPROPERTY(EditAnywhere, Category="GeometryCache", meta=(AllowedClasses = "GeometryCacheComponent"))
+	UGeometryCacheComponent* GeometryCache;
+
+	/** The offset into the beginning of the animation clip */
+	UPROPERTY(EditAnywhere, Category="GeometryCache")
+	float StartOffset;
+	
+	/** The offset into the end of the animation clip */
+	UPROPERTY(EditAnywhere, Category="GeometryCache")
+	float EndOffset;
+	
+	/** The playback rate of the animation clip */
+	UPROPERTY(EditAnywhere, Category="GeometryCache")
+	float PlayRate;
+
+	/** Reverse the playback of the animation clip */
+	UPROPERTY(EditAnywhere, Category="Animation")
+	uint32 bReverse:1;
+
+};
+
+/**
+ * Movie scene section that control geometry cache playback
+ */
+UCLASS( MinimalAPI )
+class UMovieSceneGeometryCacheSection
+	: public UMovieSceneSection
+{
+	GENERATED_UCLASS_BODY()
+
+public:
+
+	UPROPERTY(EditAnywhere, Category = "Animation", meta=(ShowOnlyInnerProperties))
+	FMovieSceneGeometryCacheParams Params;
+
+	/** Get Frame Time as Animation Time*/
+    virtual float MapTimeToAnimation(FFrameTime InPosition, FFrameRate InFrameRate) const;
+
+protected:
+	//~ UMovieSceneSection interface
+	virtual TOptional<TRange<FFrameNumber> > GetAutoSizeRange() const override;
+	virtual void TrimSection(FQualifiedFrameTime TrimTime, bool bTrimLeft) override;
+	virtual UMovieSceneSection* SplitSection(FQualifiedFrameTime SplitTime) override;
+	virtual void GetSnapTimes(TArray<FFrameNumber>& OutSnapTimes, bool bGetSectionBorders) const override;
+	virtual TOptional<FFrameTime> GetOffsetTime() const override;
+	virtual FMovieSceneEvalTemplatePtr GenerateTemplate() const override;
+
+	/** ~UObject interface */
+	virtual void Serialize(FArchive& Ar) override;
+
+private:
+
+	//~ UObject interface
+
+#if WITH_EDITOR
+
+	virtual void PreEditChange(UProperty* PropertyAboutToChange) override;
+	virtual void PostEditChangeProperty(FPropertyChangedEvent& PropertyChangedEvent) override;
+	float PreviousPlayRate;
+
+#endif
+
+
+};
diff --git a/Engine/Plugins/Experimental/GeometryCache/Source/GeometryCacheTracks/Classes/MovieSceneGeometryCacheTrack.h b/Engine/Plugins/Experimental/GeometryCache/Source/GeometryCacheTracks/Classes/MovieSceneGeometryCacheTrack.h
new file mode 100644
--- /dev/null
+++ b/Engine/Plugins/Experimental/GeometryCache/Source/GeometryCacheTracks/Classes/MovieSceneGeometryCacheTrack.h
@@ -0,0 +1,50 @@
+// Copyright 1998-2018 Epic Games, Inc. All Rights Reserved.
+
+#pragma once
+
+#include "CoreMinimal.h"
+#include "UObject/ObjectMacros.h"
+#include "MovieSceneNameableTrack.h"
+#include "MovieSceneGeometryCacheTrack.generated.h"
+
+/**
+ * Handles animation of geometry cache actors
+ */
+class UGeometryCacheComponent;
+
+UCLASS(MinimalAPI)
+class UMovieSceneGeometryCacheTrack
+	: public UMovieSceneNameableTrack
+{
+	GENERATED_UCLASS_BODY()
+
+public:
+
+	/** Adds a new animation to this track */
+	virtual UMovieSceneSection* AddNewAnimation(FFrameNumber KeyTime, UGeometryCacheComponent* GeomCacheComp);
+
+	/** Gets the animation sections at a certain time */
+	TArray<UMovieSceneSection*> GetAnimSectionsAtTime(FFrameNumber Time);
+
+public:
+
+	// UMovieSceneTrack interface
+	virtual void RemoveAllAnimationData() override;
+	virtual bool HasSection(const UMovieSceneSection& Section) const override;
+	virtual void AddSection(UMovieSceneSection& Section) override;
+	virtual void RemoveSection(UMovieSceneSection& Section) override;
+	virtual bool IsEmpty() const override;
+	virtual const TArray<UMovieSceneSection*>& GetAllSections() const override;
+	virtual UMovieSceneSection* CreateNewSection() override;
+
+#if WITH_EDITORONLY_DATA
+	virtual FText GetDefaultDisplayName() const override;
+#endif
+
+private:
+
+	/** List of all animation sections */
+	UPROPERTY()
+	TArray<UMovieSceneSection*> AnimationSections;
+
+};
diff --git a/Engine/Plugins/Experimental/GeometryCache/Source/GeometryCacheTracks/GeometryCacheTracks.Build.cs b/Engine/Plugins/Experimental/GeometryCache/Source/GeometryCacheTracks/GeometryCacheTracks.Build.cs
new file mode 100644
--- /dev/null
+++ b/Engine/Plugins/Experimental/GeometryCache/Source/GeometryCacheTracks/GeometryCacheTracks.Build.cs
@@ -0,0 +1,46 @@
+// Copyright 1998-2018 Epic Games, Inc. All Rights Reserved.
+
+namespace UnrealBuildTool.Rules
+{
+	public class GeometryCacheTracks : ModuleRules
+	{
+		public GeometryCacheTracks(ReadOnlyTargetRules Target) : base(Target)
+		{
+            OptimizeCode = CodeOptimization.Never;
+
+            PublicDependencyModuleNames.AddRange(
+				new string[]
+				{
+					"Core",
+					"CoreUObject",
+					"Engine",
+					"MovieScene",
+					"MovieSceneTracks",
+                    "GeometryCache",
+
+                }
+            );
+
+			PrivateDependencyModuleNames.AddRange(
+				new string[]
+				{
+					"Core",
+					"CoreUObject",
+					"Engine",
+                    "AnimGraphRuntime",
+                    "TimeManagement",
+					"Sequencer"
+				}
+			);
+
+
+            PublicIncludePathModuleNames.Add("TargetPlatform");
+
+            if (Target.bBuildEditor)
+            {
+                PublicIncludePathModuleNames.Add("GeometryCacheSequencer");
+                DynamicallyLoadedModuleNames.Add("GeometryCacheSequencer");
+            }
+        }
+	}
+}
diff --git a/Engine/Plugins/Experimental/GeometryCache/Source/GeometryCacheTracks/Private/GeometryCacheTracksModule.cpp b/Engine/Plugins/Experimental/GeometryCache/Source/GeometryCacheTracks/Private/GeometryCacheTracksModule.cpp
new file mode 100644
--- /dev/null
+++ b/Engine/Plugins/Experimental/GeometryCache/Source/GeometryCacheTracks/Private/GeometryCacheTracksModule.cpp
@@ -0,0 +1,19 @@
+// Copyright 1998-2018 Epic Games, Inc. All Rights Reserved.
+
+#include "GeometryCacheTracksModule.h"
+#if WITH_EDITOR
+#include "GeometryCacheSequencerModule.h"
+#endif // WITH_EDITOR
+
+IMPLEMENT_MODULE(FGeometryCacheTracksModule, GeometryCache)
+
+void FGeometryCacheTracksModule::StartupModule()
+{
+#if WITH_EDITOR
+	FGeometryCacheSequencerModule& Module = FModuleManager::LoadModuleChecked<FGeometryCacheSequencerModule>(TEXT("GeometryCacheSequencer"));
+#endif
+}
+
+void FGeometryCacheTracksModule::ShutdownModule()
+{
+}
diff --git a/Engine/Plugins/Experimental/GeometryCache/Source/GeometryCacheTracks/Private/MovieSceneGeometryCacheSection.cpp b/Engine/Plugins/Experimental/GeometryCache/Source/GeometryCacheTracks/Private/MovieSceneGeometryCacheSection.cpp
new file mode 100644
--- /dev/null
+++ b/Engine/Plugins/Experimental/GeometryCache/Source/GeometryCacheTracks/Private/MovieSceneGeometryCacheSection.cpp
@@ -0,0 +1,167 @@
+// Copyright 1998-2018 Epic Games, Inc. All Rights Reserved.
+
+#include "MovieSceneGeometryCacheSection.h"
+#include "MovieSceneGeometryCacheTemplate.h"
+#include "Logging/MessageLog.h"
+#include "MovieScene.h"
+#include "UObject/SequencerObjectVersion.h"
+#include "MovieSceneTimeHelpers.h"
+
+#define LOCTEXT_NAMESPACE "MovieSceneGeometryCacheSection"
+
+
+FMovieSceneGeometryCacheParams::FMovieSceneGeometryCacheParams()
+{
+	GeometryCache = nullptr;
+	StartOffset = 0.f;
+	EndOffset = 0.f;
+	PlayRate = 1.f;
+	bReverse = false;
+}
+
+UMovieSceneGeometryCacheSection::UMovieSceneGeometryCacheSection( const FObjectInitializer& ObjectInitializer )
+	: Super( ObjectInitializer )
+{
+
+	BlendType = EMovieSceneBlendType::Absolute;
+	EvalOptions.EnableAndSetCompletionMode(EMovieSceneCompletionMode::ProjectDefault);
+
+#if WITH_EDITOR
+
+	PreviousPlayRate = Params.PlayRate;
+
+#endif
+}
+
+TOptional<FFrameTime> UMovieSceneGeometryCacheSection::GetOffsetTime() const
+{
+	return TOptional<FFrameTime>(Params.StartOffset * GetTypedOuter<UMovieScene>()->GetTickResolution());
+}
+
+void UMovieSceneGeometryCacheSection::Serialize(FArchive& Ar)
+{
+	Ar.UsingCustomVersion(FSequencerObjectVersion::GUID);
+	Super::Serialize(Ar);
+}
+
+
+FMovieSceneEvalTemplatePtr UMovieSceneGeometryCacheSection::GenerateTemplate() const
+{
+	return FMovieSceneGeometryCacheSectionTemplate(*this);
+}
+
+float GetStartOffsetAtTrimTime(FQualifiedFrameTime TrimTime, const FMovieSceneGeometryCacheParams& Params, FFrameNumber StartFrame)
+{
+	float AnimPlayRate = FMath::IsNearlyZero(Params.PlayRate) ? 1.0f : Params.PlayRate;
+	float AnimPosition = (TrimTime.Time - StartFrame) / TrimTime.Rate * AnimPlayRate;
+	float SeqLength = Params.GetSequenceLength() - (Params.StartOffset + Params.EndOffset);
+
+	float NewOffset = FMath::Fmod(AnimPosition, SeqLength);
+	NewOffset += Params.StartOffset;
+
+	return NewOffset;
+}
+
+
+TOptional<TRange<FFrameNumber> > UMovieSceneGeometryCacheSection::GetAutoSizeRange() const
+{
+	FFrameRate FrameRate = GetTypedOuter<UMovieScene>()->GetTickResolution();
+
+	FFrameTime AnimationLength = Params.GetSequenceLength() * FrameRate;
+
+	return TRange<FFrameNumber>(GetInclusiveStartFrame(), GetInclusiveStartFrame() + AnimationLength.FrameNumber);
+}
+
+
+void UMovieSceneGeometryCacheSection::TrimSection(FQualifiedFrameTime TrimTime, bool bTrimLeft)
+{
+	SetFlags(RF_Transactional);
+
+	if (TryModify())
+	{
+		if (bTrimLeft)
+		{
+			Params.StartOffset = HasStartFrame() ? GetStartOffsetAtTrimTime(TrimTime, Params, GetInclusiveStartFrame()) : 0;
+		}
+
+		Super::TrimSection(TrimTime, bTrimLeft);
+	}
+}
+
+UMovieSceneSection* UMovieSceneGeometryCacheSection::SplitSection(FQualifiedFrameTime SplitTime)
+{
+	const float NewOffset = HasStartFrame() ? GetStartOffsetAtTrimTime(SplitTime, Params, GetInclusiveStartFrame()) : 0;
+
+	UMovieSceneSection* NewSection = Super::SplitSection(SplitTime);
+	if (NewSection != nullptr)
+	{
+		UMovieSceneGeometryCacheSection* NewGeometrySection = Cast<UMovieSceneGeometryCacheSection>(NewSection);
+		NewGeometrySection->Params.StartOffset = NewOffset;
+	}
+	return NewSection;
+}
+
+
+void UMovieSceneGeometryCacheSection::GetSnapTimes(TArray<FFrameNumber>& OutSnapTimes, bool bGetSectionBorders) const
+{
+	Super::GetSnapTimes(OutSnapTimes, bGetSectionBorders);
+
+	const FFrameRate   FrameRate  = GetTypedOuter<UMovieScene>()->GetTickResolution();
+	const FFrameNumber StartFrame = GetInclusiveStartFrame();
+	const FFrameNumber EndFrame   = GetExclusiveEndFrame() - 1; // -1 because we don't need to add the end frame twice
+
+	const float AnimPlayRate     = FMath::IsNearlyZero(Params.PlayRate) ? 1.0f : Params.PlayRate;
+	const float SeqLengthSeconds = (Params.GetSequenceLength() - (Params.StartOffset + Params.EndOffset)) / AnimPlayRate;
+
+	FFrameTime SequenceFrameLength = SeqLengthSeconds * FrameRate;
+	if (SequenceFrameLength.FrameNumber > 1)
+	{
+		// Snap to the repeat times
+		FFrameTime CurrentTime = StartFrame;
+		while (CurrentTime < EndFrame)
+		{
+			OutSnapTimes.Add(CurrentTime.FrameNumber);
+			CurrentTime += SequenceFrameLength;
+		}
+	}
+}
+
+float UMovieSceneGeometryCacheSection::MapTimeToAnimation(FFrameTime InPosition, FFrameRate InFrameRate) const
+{
+	FMovieSceneGeometryCacheSectionTemplateParameters TemplateParams(Params, GetInclusiveStartFrame(), GetExclusiveEndFrame());
+	return TemplateParams.MapTimeToAnimation(InPosition, InFrameRate);
+}
+
+
+#if WITH_EDITOR
+void UMovieSceneGeometryCacheSection::PreEditChange(UProperty* PropertyAboutToChange)
+{
+	// Store the current play rate so that we can compute the amount to compensate the section end time when the play rate changes
+	PreviousPlayRate = Params.PlayRate;
+
+	Super::PreEditChange(PropertyAboutToChange);
+}
+
+void UMovieSceneGeometryCacheSection::PostEditChangeProperty(FPropertyChangedEvent& PropertyChangedEvent)
+{
+	// Adjust the duration automatically if the play rate changes
+	if (PropertyChangedEvent.Property != nullptr &&
+		PropertyChangedEvent.Property->GetFName() == TEXT("PlayRate"))
+	{
+		float NewPlayRate = Params.PlayRate;
+
+		if (!FMath::IsNearlyZero(NewPlayRate))
+		{
+			float CurrentDuration = MovieScene::DiscreteSize(GetRange());
+			float NewDuration = CurrentDuration * (PreviousPlayRate / NewPlayRate);
+			SetEndFrame( GetInclusiveStartFrame() + FMath::FloorToInt(NewDuration) );
+
+			PreviousPlayRate = NewPlayRate;
+		}
+	}
+
+	Super::PostEditChangeProperty(PropertyChangedEvent);
+}
+#endif
+
+#undef LOCTEXT_NAMESPACE 
diff --git a/Engine/Plugins/Experimental/GeometryCache/Source/GeometryCacheTracks/Private/MovieSceneGeometryCacheTemplate.cpp b/Engine/Plugins/Experimental/GeometryCache/Source/GeometryCacheTracks/Private/MovieSceneGeometryCacheTemplate.cpp
new file mode 100644
--- /dev/null
+++ b/Engine/Plugins/Experimental/GeometryCache/Source/GeometryCacheTracks/Private/MovieSceneGeometryCacheTemplate.cpp
@@ -0,0 +1,107 @@
+// Copyright 1998-2018 Epic Games, Inc. All Rights Reserved.
+
+#include "MovieSceneGeometryCacheTemplate.h"
+#include "Compilation/MovieSceneCompilerRules.h"
+#include "GeometryCacheComponent.h"
+#include "Evaluation/MovieSceneEvaluation.h"
+#include "IMovieScenePlayer.h"
+#include "UObject/ObjectKey.h"
+#include "GeometryCacheComponent.h"
+
+
+DECLARE_CYCLE_STAT(TEXT("Geometry Cache Evaluate"), MovieSceneEval_GeometryCache_Evaluate, STATGROUP_MovieSceneEval);
+DECLARE_CYCLE_STAT(TEXT("Geometry Cache Token Execute"), MovieSceneEval_GeometryCache_TokenExecute, STATGROUP_MovieSceneEval);
+
+/** Used to set Manual Tick back to previous when outside section */
+struct FPreAnimatedGeometryCacheTokenProducer : IMovieScenePreAnimatedTokenProducer
+{
+	virtual IMovieScenePreAnimatedTokenPtr CacheExistingState(UObject& Object) const
+	{
+		struct FToken : IMovieScenePreAnimatedToken
+		{
+			FToken(UGeometryCacheComponent* InComponent)
+			{
+				// Cache this object's current update flag and animation mode
+				bInManualTick = InComponent->GetManualTick();
+			}
+
+			virtual void RestoreState(UObject& ObjectToRestore, IMovieScenePlayer& Player)
+			{
+				UGeometryCacheComponent* Component = CastChecked<UGeometryCacheComponent>(&ObjectToRestore);
+				Component->SetManualTick(bInManualTick);
+			}
+			bool bInManualTick;
+		};
+
+		return FToken(CastChecked<UGeometryCacheComponent>(&Object));
+	}
+	static FMovieSceneAnimTypeID GetAnimTypeID()
+	{
+		return TMovieSceneAnimTypeID<FPreAnimatedGeometryCacheTokenProducer>();
+	}
+};
+
+
+/** A movie scene execution token that executes a geometry cache */
+struct FGeometryCacheExecutionToken
+	: IMovieSceneExecutionToken
+	
+{
+	FGeometryCacheExecutionToken(const FMovieSceneGeometryCacheSectionTemplateParameters &InParams):
+		Params(InParams)
+	{}
+
+	/** Execute this token, operating on all objects referenced by 'Operand' */
+	virtual void Execute(const FMovieSceneContext& Context, const FMovieSceneEvaluationOperand& Operand, FPersistentEvaluationData& PersistentData, IMovieScenePlayer& Player) override
+	{
+		if (Params.GeometryCache)
+		{
+			MOVIESCENE_DETAILED_SCOPE_CYCLE_COUNTER(MovieSceneEval_GeometryCache_TokenExecute)
+
+			Player.SavePreAnimatedState(*Params.GeometryCache, FPreAnimatedGeometryCacheTokenProducer::GetAnimTypeID(), FPreAnimatedGeometryCacheTokenProducer());
+			Params.GeometryCache->SetManualTick(true);
+			// calculate the time at which to evaluate the animation
+			float EvalTime = Params.MapTimeToAnimation(Context.GetTime(), Context.GetFrameRate());
+			Params.GeometryCache->TickAtThisTime(EvalTime, true, Params.bReverse, true);
+		}
+
+	}
+
+	FMovieSceneGeometryCacheSectionTemplateParameters Params;
+};
+
+FMovieSceneGeometryCacheSectionTemplate::FMovieSceneGeometryCacheSectionTemplate(const UMovieSceneGeometryCacheSection& InSection)
+	: Params(InSection.Params, InSection.GetInclusiveStartFrame(), InSection.GetExclusiveEndFrame())
+{
+}
+
+//We use a token here so we can set the manual tick state back to what it was previously when outside this section.
+//This is similar to how Skeletal Animation evaluation also works.
+void FMovieSceneGeometryCacheSectionTemplate::Evaluate(const FMovieSceneEvaluationOperand& Operand, const FMovieSceneContext& Context, const FPersistentEvaluationData& PersistentData, FMovieSceneExecutionTokens& ExecutionTokens) const
+{
+	MOVIESCENE_DETAILED_SCOPE_CYCLE_COUNTER(MovieSceneEval_GeometryCache_Evaluate)
+	ExecutionTokens.Add(FGeometryCacheExecutionToken(Params));
+}
+
+float FMovieSceneGeometryCacheSectionTemplateParameters::MapTimeToAnimation(FFrameTime InPosition, FFrameRate InFrameRate) const
+{
+	InPosition = FMath::Clamp(InPosition, FFrameTime(SectionStartTime), FFrameTime(SectionEndTime-1));
+
+	const float SectionPlayRate = PlayRate;
+	const float AnimPlayRate = FMath::IsNearlyZero(SectionPlayRate) ? 1.0f : SectionPlayRate;
+
+	const float SeqLength = GetSequenceLength() - (StartOffset + EndOffset);
+
+	float AnimPosition = FFrameTime::FromDecimal((InPosition - SectionStartTime).AsDecimal() * AnimPlayRate) / InFrameRate;
+	if (SeqLength > 0.f)
+	{
+		AnimPosition = FMath::Fmod(AnimPosition, SeqLength);
+	}
+	AnimPosition += StartOffset;
+	if (bReverse)
+	{
+		AnimPosition = (SeqLength - (AnimPosition - StartOffset)) + StartOffset;
+	}
+
+	return AnimPosition;
+}
diff --git a/Engine/Plugins/Experimental/GeometryCache/Source/GeometryCacheTracks/Private/MovieSceneGeometryCacheTemplate.h b/Engine/Plugins/Experimental/GeometryCache/Source/GeometryCacheTracks/Private/MovieSceneGeometryCacheTemplate.h
new file mode 100644
--- /dev/null
+++ b/Engine/Plugins/Experimental/GeometryCache/Source/GeometryCacheTracks/Private/MovieSceneGeometryCacheTemplate.h
@@ -0,0 +1,45 @@
+// Copyright 1998-2018 Epic Games, Inc. All Rights Reserved.
+
+#pragma once
+
+#include "CoreMinimal.h"
+#include "UObject/ObjectMacros.h"
+#include "Evaluation/MovieSceneEvalTemplate.h"
+#include "MovieSceneGeometryCacheSection.h"
+#include "MovieSceneGeometryCacheTemplate.generated.h"
+
+USTRUCT()
+struct FMovieSceneGeometryCacheSectionTemplateParameters : public FMovieSceneGeometryCacheParams
+{
+	GENERATED_BODY()
+
+	FMovieSceneGeometryCacheSectionTemplateParameters() {}
+	FMovieSceneGeometryCacheSectionTemplateParameters(const FMovieSceneGeometryCacheParams& BaseParams, FFrameNumber InSectionStartTime, FFrameNumber InSectionEndTime)
+		: FMovieSceneGeometryCacheParams(BaseParams)
+		, SectionStartTime(InSectionStartTime)
+		, SectionEndTime(InSectionEndTime)
+	{}
+	virtual ~FMovieSceneGeometryCacheSectionTemplateParameters() {};
+	virtual float MapTimeToAnimation(FFrameTime InPosition, FFrameRate InFrameRate) const;
+
+	UPROPERTY()
+	FFrameNumber SectionStartTime;
+
+	UPROPERTY()
+	FFrameNumber SectionEndTime;
+};
+
+USTRUCT()
+struct FMovieSceneGeometryCacheSectionTemplate : public FMovieSceneEvalTemplate
+{
+	GENERATED_BODY()
+	
+	FMovieSceneGeometryCacheSectionTemplate() {}
+	FMovieSceneGeometryCacheSectionTemplate(const UMovieSceneGeometryCacheSection& Section);
+
+	virtual UScriptStruct& GetScriptStructImpl() const override { return *StaticStruct(); }
+	virtual void Evaluate(const FMovieSceneEvaluationOperand& Operand, const FMovieSceneContext& Context, const FPersistentEvaluationData& PersistentData, FMovieSceneExecutionTokens& ExecutionTokens) const override;
+
+	UPROPERTY()
+	FMovieSceneGeometryCacheSectionTemplateParameters Params;
+};
diff --git a/Engine/Plugins/Experimental/GeometryCache/Source/GeometryCacheTracks/Private/MovieSceneGeometryCacheTrack.cpp b/Engine/Plugins/Experimental/GeometryCache/Source/GeometryCacheTracks/Private/MovieSceneGeometryCacheTrack.cpp
new file mode 100644
--- /dev/null
+++ b/Engine/Plugins/Experimental/GeometryCache/Source/GeometryCacheTracks/Private/MovieSceneGeometryCacheTrack.cpp
@@ -0,0 +1,120 @@
+// Copyright 1998-2018 Epic Games, Inc. All Rights Reserved.
+
+#include "MovieSceneGeometryCacheTrack.h"
+#include "GeometryCacheComponent.h"
+#include "MovieSceneGeometryCacheSection.h"
+#include "Compilation/MovieSceneCompilerRules.h"
+#include "Evaluation/MovieSceneEvaluationTrack.h"
+#include "MovieSceneGeometryCacheTemplate.h"
+#include "Compilation/IMovieSceneTemplateGenerator.h"
+#include "MovieScene.h"
+
+#define LOCTEXT_NAMESPACE "MovieSceneSkeletalAnimationTrack"
+
+
+/* UMovieSceneGeometryCacheTrack structors
+ *****************************************************************************/
+
+UMovieSceneGeometryCacheTrack::UMovieSceneGeometryCacheTrack(const FObjectInitializer& ObjectInitializer)
+	: Super(ObjectInitializer)
+{
+#if WITH_EDITORONLY_DATA
+	TrackTint = FColor(124, 15, 124, 65);
+#endif
+
+	SupportedBlendTypes.Add(EMovieSceneBlendType::Absolute);
+
+	EvalOptions.bCanEvaluateNearestSection = true;
+	EvalOptions.bEvaluateInPreroll = true;
+}
+
+
+/* UMovieSceneGeometryCacheTrack interface
+ *****************************************************************************/
+
+UMovieSceneSection* UMovieSceneGeometryCacheTrack::AddNewAnimation(FFrameNumber KeyTime, UGeometryCacheComponent* GeomCacheComp)
+{
+	UMovieSceneGeometryCacheSection* NewSection = Cast<UMovieSceneGeometryCacheSection>(CreateNewSection());
+	{
+		FFrameTime AnimationLength = GeomCacheComp->GetDuration()* GetTypedOuter<UMovieScene>()->GetTickResolution();
+		NewSection->InitialPlacementOnRow(AnimationSections, KeyTime, AnimationLength.FrameNumber.Value, INDEX_NONE);
+		NewSection->Params.GeometryCache = GeomCacheComp;
+	}
+
+	AddSection(*NewSection);
+
+	return NewSection;
+}
+
+
+TArray<UMovieSceneSection*> UMovieSceneGeometryCacheTrack::GetAnimSectionsAtTime(FFrameNumber Time)
+{
+	TArray<UMovieSceneSection*> Sections;
+	for (auto Section : AnimationSections)
+	{
+		if (Section->IsTimeWithinSection(Time))
+		{
+			Sections.Add(Section);
+		}
+	}
+
+	return Sections;
+}
+
+
+/* UMovieSceneTrack interface
+ *****************************************************************************/
+
+
+const TArray<UMovieSceneSection*>& UMovieSceneGeometryCacheTrack::GetAllSections() const
+{
+	return AnimationSections;
+}
+
+
+UMovieSceneSection* UMovieSceneGeometryCacheTrack::CreateNewSection()
+{
+	return NewObject<UMovieSceneGeometryCacheSection>(this, NAME_None, RF_Transactional);
+}
+
+
+void UMovieSceneGeometryCacheTrack::RemoveAllAnimationData()
+{
+	AnimationSections.Empty();
+}
+
+
+bool UMovieSceneGeometryCacheTrack::HasSection(const UMovieSceneSection& Section) const
+{
+	return AnimationSections.Contains(&Section);
+}
+
+
+void UMovieSceneGeometryCacheTrack::AddSection(UMovieSceneSection& Section)
+{
+	AnimationSections.Add(&Section);
+}
+
+
+void UMovieSceneGeometryCacheTrack::RemoveSection(UMovieSceneSection& Section)
+{
+	AnimationSections.Remove(&Section);
+}
+
+
+bool UMovieSceneGeometryCacheTrack::IsEmpty() const
+{
+	return AnimationSections.Num() == 0;
+}
+
+#if WITH_EDITORONLY_DATA
+
+FText UMovieSceneGeometryCacheTrack::GetDefaultDisplayName() const
+{
+	return LOCTEXT("TrackName", "Geometry Cache");
+}
+
+#endif
+
+
+#undef LOCTEXT_NAMESPACE
diff --git a/Engine/Plugins/Experimental/GeometryCache/Source/GeometryCacheTracks/Public/GeometryCacheTracksModule.h b/Engine/Plugins/Experimental/GeometryCache/Source/GeometryCacheTracks/Public/GeometryCacheTracksModule.h
new file mode 100644
--- /dev/null
+++ b/Engine/Plugins/Experimental/GeometryCache/Source/GeometryCacheTracks/Public/GeometryCacheTracksModule.h
@@ -0,0 +1,42 @@
+// Copyright 1998-2018 Epic Games, Inc. All Rights Reserved.
+
+#pragma once
+
+#include "CoreMinimal.h"
+#include "Modules/ModuleInterface.h"
+#include "Modules/ModuleManager.h"
+
+/**
+ * The public interface to this module
+ */
+class FGeometryCacheTracksModule : public IModuleInterface
+{
+public:
+	FGeometryCacheTracksModule() {}
+
+	virtual void StartupModule();
+	virtual void ShutdownModule();
+
+	/**
+	 * Singleton-like access to this module's interface.  This is just for convenience!
+	 * Beware of calling this during the shutdown phase, though.  Your module might have been unloaded already.
+	 *
+	 * @return Returns singleton instance, loading the module on demand if needed
+	 */
+	static inline FGeometryCacheTracksModule& Get()
+	{
+		return FModuleManager::LoadModuleChecked< FGeometryCacheTracksModule >("GeometryCacheTracks");
+	}
+
+	/**
+	 * Checks to see if this module is loaded and ready.  It is only valid to call Get() if IsAvailable() returns true.
+	 *
+	 * @return True if the module is loaded and ready to use
+	 */
+	static inline bool IsAvailable()
+	{
+		return FModuleManager::Get().IsModuleLoaded( "GeometryCacheTracks" );
+	}
+private:
+};
+
