Fix CIS

diff --git a/Engine/Plugins/Experimental/GeometryCache/Source/GeometryCacheTracks/GeometryCacheTracks.Build.cs b/Engine/Plugins/Experimental/GeometryCache/Source/GeometryCacheTracks/GeometryCacheTracks.Build.cs
--- a/Engine/Plugins/Experimental/GeometryCache/Source/GeometryCacheTracks/GeometryCacheTracks.Build.cs
+++ b/Engine/Plugins/Experimental/GeometryCache/Source/GeometryCacheTracks/GeometryCacheTracks.Build.cs
@@ -29,7 +29,6 @@ namespace UnrealBuildTool.Rules
 					"Engine",
                     "AnimGraphRuntime",
                     "TimeManagement",
-					"Sequencer"
 				}
 			);
 
