Fix CIS

diff --git a/Engine/Plugins/Experimental/GeometryCache/Source/GeometryCacheSequencer/GeometryCacheSequencer.Build.cs b/Engine/Plugins/Experimental/GeometryCache/Source/GeometryCacheSequencer/GeometryCacheSequencer.Build.cs
--- a/Engine/Plugins/Experimental/GeometryCache/Source/GeometryCacheSequencer/GeometryCacheSequencer.Build.cs
+++ b/Engine/Plugins/Experimental/GeometryCache/Source/GeometryCacheSequencer/GeometryCacheSequencer.Build.cs
@@ -4,23 +4,30 @@ namespace UnrealBuildTool.Rules
 {
 	public class GeometryCacheSequencer : ModuleRules
 	{
-		public GeometryCacheSequencer(ReadOnlyTargetRules Target) : base(Target)
-		{
+        public GeometryCacheSequencer(ReadOnlyTargetRules Target) : base(Target)
+        {
             OptimizeCode = CodeOptimization.Never;
 
+            PublicDependencyModuleNames.AddRange(
+                new string[]
+                {
+                    "Sequencer",
+                }
+            );
+
+
             PrivateDependencyModuleNames.AddRange(
-				new string[]
-				{
-					"AssetTools",
-					"Core",
-					"CoreUObject",
+                new string[]
+                {
+                    "AssetTools",
+                    "Core",
+                    "CoreUObject",
 					"EditorStyle",
 					"Engine",
 					"MovieScene",
 					"MovieSceneTools",
 					"MovieSceneTracks",
 					"RHI",
-					"Sequencer",
 					"Slate",
 					"SlateCore",
                     "TimeManagement",
