Fix CIS

diff --git a/Engine/Plugins/Experimental/GeometryCache/Source/GeometryCacheSequencer/GeometryCacheSequencer.Build.cs b/Engine/Plugins/Experimental/GeometryCache/Source/GeometryCacheSequencer/GeometryCacheSequencer.Build.cs
--- a/Engine/Plugins/Experimental/GeometryCache/Source/GeometryCacheSequencer/GeometryCacheSequencer.Build.cs
+++ b/Engine/Plugins/Experimental/GeometryCache/Source/GeometryCacheSequencer/GeometryCacheSequencer.Build.cs
@@ -6,16 +6,11 @@ namespace UnrealBuildTool.Rules
 	{
         public GeometryCacheSequencer(ReadOnlyTargetRules Target) : base(Target)
         {
-            OptimizeCode = CodeOptimization.Never;
-
-            PublicDependencyModuleNames.AddRange(
-                new string[]
-                {
-                    "Sequencer",
+            PublicIncludePathModuleNames.AddRange(
+                new string[] {
+                "Sequencer",
                 }
             );
-
-
             PrivateDependencyModuleNames.AddRange(
                 new string[]
                 {
@@ -28,6 +23,7 @@ namespace UnrealBuildTool.Rules
 					"MovieSceneTools",
 					"MovieSceneTracks",
 					"RHI",
+                    "Sequencer",
 					"Slate",
 					"SlateCore",
                     "TimeManagement",
