Turn off turning off optimization for geometry sequencer tracks.

#trivial

diff --git a/Engine/Plugins/Experimental/GeometryCache/Source/GeometryCacheTracks/GeometryCacheTracks.Build.cs b/Engine/Plugins/Experimental/GeometryCache/Source/GeometryCacheTracks/GeometryCacheTracks.Build.cs
--- a/Engine/Plugins/Experimental/GeometryCache/Source/GeometryCacheTracks/GeometryCacheTracks.Build.cs
+++ b/Engine/Plugins/Experimental/GeometryCache/Source/GeometryCacheTracks/GeometryCacheTracks.Build.cs
@@ -6,8 +6,6 @@ namespace UnrealBuildTool.Rules
 	{
 		public GeometryCacheTracks(ReadOnlyTargetRules Target) : base(Target)
 		{
-            OptimizeCode = CodeOptimization.Never;
-
             PublicDependencyModuleNames.AddRange(
 				new string[]
 				{
