Fix CIS

diff --git a/Engine/Plugins/Experimental/GeometryCache/Source/GeometryCacheTracks/Private/GeometryCacheTracksModule.cpp b/Engine/Plugins/Experimental/GeometryCache/Source/GeometryCacheTracks/Private/GeometryCacheTracksModule.cpp
--- a/Engine/Plugins/Experimental/GeometryCache/Source/GeometryCacheTracks/Private/GeometryCacheTracksModule.cpp
+++ b/Engine/Plugins/Experimental/GeometryCache/Source/GeometryCacheTracks/Private/GeometryCacheTracksModule.cpp
@@ -5,7 +5,7 @@
 #include "GeometryCacheSequencerModule.h"
 #endif // WITH_EDITOR
 
-IMPLEMENT_MODULE(FGeometryCacheTracksModule, GeometryCache)
+IMPLEMENT_MODULE(FGeometryCacheTracksModule, GeometryCacheTracks)
 
 void FGeometryCacheTracksModule::StartupModule()
 {
