Sequencer: Change geometry cache component to soft object path to fix GLEOs

diff --git a/Engine/Plugins/Experimental/GeometryCache/Source/GeometryCacheSequencer/Classes/GeometryCacheTrackEditor.cpp b/Engine/Plugins/Experimental/GeometryCache/Source/GeometryCacheSequencer/Classes/GeometryCacheTrackEditor.cpp
--- a/Engine/Plugins/Experimental/GeometryCache/Source/GeometryCacheSequencer/Classes/GeometryCacheTrackEditor.cpp
+++ b/Engine/Plugins/Experimental/GeometryCache/Source/GeometryCacheSequencer/Classes/GeometryCacheTrackEditor.cpp
@@ -83,9 +83,13 @@ UMovieSceneSection* FGeometryCacheSection::GetSectionObject()
 
 FText FGeometryCacheSection::GetSectionTitle() const
 {
-	if (Section.Params.GeometryCache != nullptr  && Section.Params.GeometryCache->GetOwner() != nullptr)
+	if (Section.Params.GeometryCache.ResolveObject() != nullptr )
 	{
-		return FText::FromString(Section.Params.GeometryCache->GetOwner()->GetName());
+		UGeometryCacheComponent* GeometryCache = Cast<UGeometryCacheComponent>(Section.Params.GeometryCache.ResolveObject());
+		if (GeometryCache && GeometryCache->GetOwner() != nullptr)
+		{	
+			return FText::FromString(GeometryCache->GetOwner()->GetName());
+		}
 	}
 	return LOCTEXT("NoGeometryCacheSection", "No GeometryCache");
 }
@@ -146,12 +150,15 @@ int32 FGeometryCacheSection::OnPaintSection( FSequencerSectionPainter& Painter )
 	if (Painter.bIsSelected && SequencerPtr.IsValid())
 	{
 		FFrameTime CurrentTime = SequencerPtr->GetLocalTime().Time;
-		if (Section.GetRange().Contains(CurrentTime.FrameNumber) && Section.Params.GeometryCache != nullptr)
+		if (Section.GetRange().Contains(CurrentTime.FrameNumber) && Section.Params.GeometryCache.ResolveObject() != nullptr)
 		{
 			const float Time = TimeToPixelConverter.FrameToPixel(CurrentTime); 
 
+			UGeometryCacheComponent* GeometryCache = Cast<UGeometryCacheComponent>(Section.Params.GeometryCache.ResolveObject());
+
 			// Draw the current time next to the scrub handle
-			const float AnimTime = Section.MapTimeToAnimation(CurrentTime, TickResolution);			int32 FrameTime = Section.Params.GeometryCache->GetFrameAtTime(AnimTime);
+			const float AnimTime = Section.MapTimeToAnimation(CurrentTime, TickResolution);			
+			int32 FrameTime = GeometryCache->GetFrameAtTime(AnimTime);
 			FString FrameString = FString::FromInt(FrameTime);
 
 			const FSlateFontInfo SmallLayoutFont = FCoreStyle::GetDefaultFontStyle("Bold", 10);
diff --git a/Engine/Plugins/Experimental/GeometryCache/Source/GeometryCacheTracks/Classes/MovieSceneGeometryCacheSection.h b/Engine/Plugins/Experimental/GeometryCache/Source/GeometryCacheTracks/Classes/MovieSceneGeometryCacheSection.h
--- a/Engine/Plugins/Experimental/GeometryCache/Source/GeometryCacheTracks/Classes/MovieSceneGeometryCacheSection.h
+++ b/Engine/Plugins/Experimental/GeometryCache/Source/GeometryCacheTracks/Classes/MovieSceneGeometryCacheSection.h
@@ -7,6 +7,7 @@
 #include "MovieSceneSection.h"
 #include "GeometryCacheComponent.h"
 #include "Channels/MovieSceneFloatChannel.h"
+#include "UObject/SoftObjectPath.h"
 #include "MovieSceneGeometryCacheSection.generated.h"
 
 USTRUCT()
@@ -17,14 +18,14 @@ struct FMovieSceneGeometryCacheParams
 	FMovieSceneGeometryCacheParams();
 
 	/** Gets the animation duration, modified by play rate */
-	float GetDuration() const { return FMath::IsNearlyZero(PlayRate) || GeometryCache == nullptr ? 0.f : GeometryCache->GetDuration() / PlayRate; }
+	float GetDuration() const { return FMath::IsNearlyZero(PlayRate) || GeometryCache.ResolveObject() == nullptr ? 0.f : Cast<UGeometryCacheComponent>(GeometryCache.ResolveObject())->GetDuration() / PlayRate; }
 
 	/** Gets the animation sequence length, not modified by play rate */
-	float GetSequenceLength() const { return GeometryCache != nullptr ? GeometryCache->GetDuration() : 0.f; }
+	float GetSequenceLength() const { return GeometryCache.ResolveObject() != nullptr ? Cast<UGeometryCacheComponent>(GeometryCache.ResolveObject())->GetDuration() : 0.f; }
 
 	/** The animation this section plays */
 	UPROPERTY(EditAnywhere, Category="GeometryCache", meta=(AllowedClasses = "GeometryCacheComponent"))
-	UGeometryCacheComponent* GeometryCache;
+	FSoftObjectPath GeometryCache;
 
 	/** The offset into the beginning of the animation clip */
 	UPROPERTY(EditAnywhere, Category="GeometryCache")
diff --git a/Engine/Plugins/Experimental/GeometryCache/Source/GeometryCacheTracks/Private/MovieSceneGeometryCacheSection.cpp b/Engine/Plugins/Experimental/GeometryCache/Source/GeometryCacheTracks/Private/MovieSceneGeometryCacheSection.cpp
--- a/Engine/Plugins/Experimental/GeometryCache/Source/GeometryCacheTracks/Private/MovieSceneGeometryCacheSection.cpp
+++ b/Engine/Plugins/Experimental/GeometryCache/Source/GeometryCacheTracks/Private/MovieSceneGeometryCacheSection.cpp
@@ -12,7 +12,6 @@
 
 FMovieSceneGeometryCacheParams::FMovieSceneGeometryCacheParams()
 {
-	GeometryCache = nullptr;
 	StartOffset = 0.f;
 	EndOffset = 0.f;
 	PlayRate = 1.f;
diff --git a/Engine/Plugins/Experimental/GeometryCache/Source/GeometryCacheTracks/Private/MovieSceneGeometryCacheTemplate.cpp b/Engine/Plugins/Experimental/GeometryCache/Source/GeometryCacheTracks/Private/MovieSceneGeometryCacheTemplate.cpp
--- a/Engine/Plugins/Experimental/GeometryCache/Source/GeometryCacheTracks/Private/MovieSceneGeometryCacheTemplate.cpp
+++ b/Engine/Plugins/Experimental/GeometryCache/Source/GeometryCacheTracks/Private/MovieSceneGeometryCacheTemplate.cpp
@@ -54,15 +54,17 @@ struct FGeometryCacheExecutionToken
 	/** Execute this token, operating on all objects referenced by 'Operand' */
 	virtual void Execute(const FMovieSceneContext& Context, const FMovieSceneEvaluationOperand& Operand, FPersistentEvaluationData& PersistentData, IMovieScenePlayer& Player) override
 	{
-		if (Params.GeometryCache)
+		if (Params.GeometryCache.ResolveObject())
 		{
 			MOVIESCENE_DETAILED_SCOPE_CYCLE_COUNTER(MovieSceneEval_GeometryCache_TokenExecute)
 
-			Player.SavePreAnimatedState(*Params.GeometryCache, FPreAnimatedGeometryCacheTokenProducer::GetAnimTypeID(), FPreAnimatedGeometryCacheTokenProducer());
-			Params.GeometryCache->SetManualTick(true);
+			UGeometryCacheComponent* GeometryCache = Cast<UGeometryCacheComponent>(Params.GeometryCache.ResolveObject());
+
+			Player.SavePreAnimatedState(*GeometryCache, FPreAnimatedGeometryCacheTokenProducer::GetAnimTypeID(), FPreAnimatedGeometryCacheTokenProducer());
+			GeometryCache->SetManualTick(true);
 			// calculate the time at which to evaluate the animation
 			float EvalTime = Params.MapTimeToAnimation(Context.GetTime(), Context.GetFrameRate());
-			Params.GeometryCache->TickAtThisTime(EvalTime, true, Params.bReverse, true);
+			GeometryCache->TickAtThisTime(EvalTime, true, Params.bReverse, true);
 		}
 
 	}
