Missing forward declaration
#trivial

diff --git a/Engine/Plugins/Experimental/GeometryCache/Source/GeometryCacheTracks/Classes/MovieSceneGeometryCacheTrack.h b/Engine/Plugins/Experimental/GeometryCache/Source/GeometryCacheTracks/Classes/MovieSceneGeometryCacheTrack.h
--- a/Engine/Plugins/Experimental/GeometryCache/Source/GeometryCacheTracks/Classes/MovieSceneGeometryCacheTrack.h
+++ b/Engine/Plugins/Experimental/GeometryCache/Source/GeometryCacheTracks/Classes/MovieSceneGeometryCacheTrack.h
@@ -10,6 +10,8 @@
 /**
  * Handles animation of geometry cache actors
  */
+class UGeometryCacheComponent;
+
 UCLASS(MinimalAPI)
 class UMovieSceneGeometryCacheTrack
 	: public UMovieSceneNameableTrack
diff --git a/Engine/Plugins/Experimental/GeometryCache/Source/GeometryCacheTracks/Private/MovieSceneGeometryCacheTrack.cpp b/Engine/Plugins/Experimental/GeometryCache/Source/GeometryCacheTracks/Private/MovieSceneGeometryCacheTrack.cpp
--- a/Engine/Plugins/Experimental/GeometryCache/Source/GeometryCacheTracks/Private/MovieSceneGeometryCacheTrack.cpp
+++ b/Engine/Plugins/Experimental/GeometryCache/Source/GeometryCacheTracks/Private/MovieSceneGeometryCacheTrack.cpp
@@ -1,6 +1,7 @@
 // Copyright 1998-2018 Epic Games, Inc. All Rights Reserved.
 
 #include "MovieSceneGeometryCacheTrack.h"
+#include "GeometryCacheComponent.h"
 #include "MovieSceneGeometryCacheSection.h"
 #include "Compilation/MovieSceneCompilerRules.h"
 #include "Evaluation/MovieSceneEvaluationTrack.h"
