CIS Fixes from 4.20 from Max.

diff --git a/Engine/Plugins/Experimental/GeometryCache/Source/GeometryCacheSequencer/GeometryCacheSequencer.Build.cs b/Engine/Plugins/Experimental/GeometryCache/Source/GeometryCacheSequencer/GeometryCacheSequencer.Build.cs
--- a/Engine/Plugins/Experimental/GeometryCache/Source/GeometryCacheSequencer/GeometryCacheSequencer.Build.cs
+++ b/Engine/Plugins/Experimental/GeometryCache/Source/GeometryCacheSequencer/GeometryCacheSequencer.Build.cs
@@ -2,33 +2,36 @@
 
 namespace UnrealBuildTool.Rules
 {
-	public class GeometryCacheSequencer : ModuleRules
-	{
-		public GeometryCacheSequencer(ReadOnlyTargetRules Target) : base(Target)
-		{
-            OptimizeCode = CodeOptimization.Never;
-
+    public class GeometryCacheSequencer : ModuleRules
+    {
+        public GeometryCacheSequencer(ReadOnlyTargetRules Target) : base(Target)
+        {
+            PublicIncludePathModuleNames.AddRange(
+                new string[] {
+                "Sequencer",
+                }
+            );
             PrivateDependencyModuleNames.AddRange(
-				new string[]
-				{
-					"AssetTools",
-					"Core",
-					"CoreUObject",
-					"EditorStyle",
-					"Engine",
-					"MovieScene",
-					"MovieSceneTools",
-					"MovieSceneTracks",
-					"RHI",
-					"Sequencer",
-					"Slate",
-					"SlateCore",
+                new string[]
+                {
+                    "AssetTools",
+                    "Core",
+                    "CoreUObject",
+                    "EditorStyle",
+                    "Engine",
+                    "MovieScene",
+                    "MovieSceneTools",
+                    "MovieSceneTracks",
+                    "RHI",
+                    "Sequencer",
+                    "Slate",
+                    "SlateCore",
                     "TimeManagement",
-					"UnrealEd",
+                    "UnrealEd",
                     "GeometryCacheTracks",
                     "GeometryCache"
                 }
-			);
+            );
         }
-	}
+    }
 }
diff --git a/Engine/Plugins/Experimental/GeometryCache/Source/GeometryCacheTracks/GeometryCacheTracks.Build.cs b/Engine/Plugins/Experimental/GeometryCache/Source/GeometryCacheTracks/GeometryCacheTracks.Build.cs
--- a/Engine/Plugins/Experimental/GeometryCache/Source/GeometryCacheTracks/GeometryCacheTracks.Build.cs
+++ b/Engine/Plugins/Experimental/GeometryCache/Source/GeometryCacheTracks/GeometryCacheTracks.Build.cs
@@ -2,36 +2,33 @@
 
 namespace UnrealBuildTool.Rules
 {
-	public class GeometryCacheTracks : ModuleRules
-	{
-		public GeometryCacheTracks(ReadOnlyTargetRules Target) : base(Target)
-		{
-            OptimizeCode = CodeOptimization.Never;
-
+    public class GeometryCacheTracks : ModuleRules
+    {
+        public GeometryCacheTracks(ReadOnlyTargetRules Target) : base(Target)
+        {
             PublicDependencyModuleNames.AddRange(
-				new string[]
-				{
-					"Core",
-					"CoreUObject",
-					"Engine",
-					"MovieScene",
-					"MovieSceneTracks",
+                new string[]
+                {
+                    "Core",
+                    "CoreUObject",
+                    "Engine",
+                    "MovieScene",
+                    "MovieSceneTracks",
                     "GeometryCache",
 
                 }
             );
 
-			PrivateDependencyModuleNames.AddRange(
-				new string[]
-				{
-					"Core",
-					"CoreUObject",
-					"Engine",
+            PrivateDependencyModuleNames.AddRange(
+                new string[]
+                {
+                    "Core",
+                    "CoreUObject",
+                    "Engine",
                     "AnimGraphRuntime",
                     "TimeManagement",
-					"Sequencer"
-				}
-			);
+                }
+            );
 
 
             PublicIncludePathModuleNames.Add("TargetPlatform");
@@ -42,5 +39,5 @@ namespace UnrealBuildTool.Rules
                 DynamicallyLoadedModuleNames.Add("GeometryCacheSequencer");
             }
         }
-	}
+    }
 }
