Sequencer: Geometry cache track needs to update the Elapsed Time in order for the chunk time to update. This fixes an issue where the geometry cache track only evalutes for 5 seconds.

Copy from Release-4.20

diff --git a/Engine/Plugins/Experimental/GeometryCache/Source/GeometryCache/Private/GeometryCacheComponent.cpp b/Engine/Plugins/Experimental/GeometryCache/Source/GeometryCache/Private/GeometryCacheComponent.cpp
--- a/Engine/Plugins/Experimental/GeometryCache/Source/GeometryCache/Private/GeometryCacheComponent.cpp
+++ b/Engine/Plugins/Experimental/GeometryCache/Source/GeometryCache/Private/GeometryCacheComponent.cpp
@@ -152,6 +152,8 @@ void UGeometryCacheComponent::TickAtThisTime(const float Time, bool bInIsRunning
 {
 	if (bManualTick && GeometryCache && bRunning)
 	{
+		ElapsedTime = Time;
+
 		// Game thread update:
 		// This mainly just updates the matrix and bounding boxes. All render state (meshes) is done on the render thread
 		bool bUpdatedBoundsOrMatrix = false;
