Removed GeometryCache dependency on MeshUtilitiesCommon in non-editor configs.

diff --git a/Engine/Plugins/Experimental/GeometryCache/Source/GeometryCache/GeometryCache.Build.cs b/Engine/Plugins/Experimental/GeometryCache/Source/GeometryCache/GeometryCache.Build.cs
--- a/Engine/Plugins/Experimental/GeometryCache/Source/GeometryCache/GeometryCache.Build.cs
+++ b/Engine/Plugins/Experimental/GeometryCache/Source/GeometryCache/GeometryCache.Build.cs
@@ -20,18 +20,13 @@ public class GeometryCache : ModuleRules
 			}
 		);
 
-        PrivateDependencyModuleNames.AddRange(
-            new string[] {
-                "MeshUtilitiesCommon",
-            }
-        );
-
         PublicIncludePathModuleNames.Add("TargetPlatform");
 
         if (Target.bBuildEditor)
         {
             PublicIncludePathModuleNames.Add("GeometryCacheEd");
             DynamicallyLoadedModuleNames.Add("GeometryCacheEd");
+            PrivateDependencyModuleNames.Add("MeshUtilitiesCommon");
         }        
 	}
 }
