Fix for Geometry Cache Tracks not loading with Spanwables.

diff --git a/Engine/Plugins/Experimental/GeometryCache/Source/GeometryCacheSequencer/Classes/GeometryCacheTrackEditor.cpp b/Engine/Plugins/Experimental/GeometryCache/Source/GeometryCacheSequencer/Classes/GeometryCacheTrackEditor.cpp
--- a/Engine/Plugins/Experimental/GeometryCache/Source/GeometryCacheSequencer/Classes/GeometryCacheTrackEditor.cpp
+++ b/Engine/Plugins/Experimental/GeometryCache/Source/GeometryCacheSequencer/Classes/GeometryCacheTrackEditor.cpp
@@ -52,7 +52,7 @@ static UGeometryCacheComponent* AcquireGeometryCacheFromObjectGuid(const FGuid&
 			}
 		}
 	}
-	else if(UGeometryCacheComponent* GeometryMeshComp = Cast<UGeometryCacheComponent>(BoundObject))
+	else if (UGeometryCacheComponent* GeometryMeshComp = Cast<UGeometryCacheComponent>(BoundObject))
 	{
 		if (GeometryMeshComp->GetGeometryCache())
 		{
@@ -64,7 +64,7 @@ static UGeometryCacheComponent* AcquireGeometryCacheFromObjectGuid(const FGuid&
 }
 
 
-FGeometryCacheSection::FGeometryCacheSection( UMovieSceneSection& InSection, TWeakPtr<ISequencer> InSequencer)
+FGeometryCacheSection::FGeometryCacheSection(UMovieSceneSection& InSection, TWeakPtr<ISequencer> InSequencer)
 	: Section(*CastChecked<UMovieSceneGeometryCacheSection>(&InSection))
 	, Sequencer(InSequencer)
 	, InitialStartOffsetDuringResize(0)
@@ -73,18 +73,18 @@ FGeometryCacheSection::FGeometryCacheSection( UMovieSceneSection& InSection, TWe
 
 
 UMovieSceneSection* FGeometryCacheSection::GetSectionObject()
-{ 
+{
 	return &Section;
 }
 
 
 FText FGeometryCacheSection::GetSectionTitle() const
 {
-	if (Section.Params.GeometryCache.ResolveObject() != nullptr )
+	if (Section.Params.GeometryCache.Get() != nullptr)
 	{
-		UGeometryCacheComponent* GeometryCache = Cast<UGeometryCacheComponent>(Section.Params.GeometryCache.ResolveObject());
+		UGeometryCacheComponent* GeometryCache = Cast<UGeometryCacheComponent>(Section.Params.GeometryCache.Get());
 		if (GeometryCache && GeometryCache->GetOwner() != nullptr)
-		{	
+		{
 			return FText::FromString(GeometryCache->GetOwner()->GetName());
 		}
 	}
@@ -98,10 +98,10 @@ float FGeometryCacheSection::GetSectionHeight() const
 }
 
 
-int32 FGeometryCacheSection::OnPaintSection( FSequencerSectionPainter& Painter ) const
+int32 FGeometryCacheSection::OnPaintSection(FSequencerSectionPainter& Painter) const
 {
 	const ESlateDrawEffect DrawEffects = Painter.bParentEnabled ? ESlateDrawEffect::None : ESlateDrawEffect::DisabledEffect;
-	
+
 	const FTimeToPixel& TimeToPixelConverter = Painter.GetTimeConverter();
 
 	int32 LayerId = Painter.PaintSectionBackground();
@@ -121,9 +121,9 @@ int32 FGeometryCacheSection::OnPaintSection( FSequencerSectionPainter& Painter )
 
 	if (!FMath::IsNearlyZero(SeqLength, KINDA_SMALL_NUMBER) && SeqLength > 0)
 	{
-		float MaxOffset  = Section.GetRange().Size<FFrameTime>() / TickResolution;
+		float MaxOffset = Section.GetRange().Size<FFrameTime>() / TickResolution;
 		float OffsetTime = SeqLength;
-		float StartTime  = Section.GetInclusiveStartFrame() / TickResolution;
+		float StartTime = Section.GetInclusiveStartFrame() / TickResolution;
 
 		while (OffsetTime < MaxOffset)
 		{
@@ -133,7 +133,7 @@ int32 FGeometryCacheSection::OnPaintSection( FSequencerSectionPainter& Painter )
 				Painter.DrawElements,
 				LayerId,
 				Painter.SectionGeometry.MakeChild(
-					FVector2D(2.f, Painter.SectionGeometry.Size.Y-2.f),
+					FVector2D(2.f, Painter.SectionGeometry.Size.Y - 2.f),
 					FSlateLayoutTransform(FVector2D(OffsetPixel, 1.f))
 				).ToPaintGeometry(),
 				GenericDivider,
@@ -148,14 +148,14 @@ int32 FGeometryCacheSection::OnPaintSection( FSequencerSectionPainter& Painter )
 	if (Painter.bIsSelected && SequencerPtr.IsValid())
 	{
 		FFrameTime CurrentTime = SequencerPtr->GetLocalTime().Time;
-		if (Section.GetRange().Contains(CurrentTime.FrameNumber) && Section.Params.GeometryCache.ResolveObject() != nullptr)
+		if (Section.GetRange().Contains(CurrentTime.FrameNumber) && Section.Params.GeometryCache.Get() != nullptr)
 		{
-			const float Time = TimeToPixelConverter.FrameToPixel(CurrentTime); 
+			const float Time = TimeToPixelConverter.FrameToPixel(CurrentTime);
 
-			UGeometryCacheComponent* GeometryCache = Cast<UGeometryCacheComponent>(Section.Params.GeometryCache.ResolveObject());
+			UGeometryCacheComponent* GeometryCache = Cast<UGeometryCacheComponent>(Section.Params.GeometryCache.Get());
 
 			// Draw the current time next to the scrub handle
-			const float AnimTime = Section.MapTimeToAnimation(CurrentTime, TickResolution);			
+			const float AnimTime = Section.MapTimeToAnimation(CurrentTime, TickResolution);
 			int32 FrameTime = GeometryCache->GetFrameAtTime(AnimTime);
 			FString FrameString = FString::FromInt(FrameTime);
 
@@ -168,13 +168,13 @@ int32 FGeometryCacheSection::OnPaintSection( FSequencerSectionPainter& Painter )
 			bool  bDrawLeft = (Painter.SectionGeometry.Size.X - Time) < (TextSize.X + 22.f) - TextOffsetPx;
 			float TextPosition = bDrawLeft ? Time - TextSize.X - TextOffsetPx : Time + TextOffsetPx;
 			//handle mirrored labels
-			const float MajorTickHeight = 9.0f; 
+			const float MajorTickHeight = 9.0f;
 			FVector2D TextOffset(TextPosition, Painter.SectionGeometry.Size.Y - (MajorTickHeight + TextSize.Y));
 
 			const FLinearColor DrawColor = FEditorStyle::GetSlateColor("SelectionColor").GetColor(FWidgetStyle());
 			const FVector2D BoxPadding = FVector2D(4.0f, 2.0f);
 			// draw time string
-	
+
 			FSlateDrawElement::MakeBox(
 				Painter.DrawElements,
 				LayerId + 5,
@@ -196,14 +196,14 @@ int32 FGeometryCacheSection::OnPaintSection( FSequencerSectionPainter& Painter )
 
 		}
 	}
-	
+
 	return LayerId;
 }
 
 void FGeometryCacheSection::BeginResizeSection()
 {
 	InitialStartOffsetDuringResize = Section.Params.StartFrameOffset;
-	InitialStartTimeDuringResize   = Section.HasStartFrame() ? Section.GetInclusiveStartFrame() : 0;
+	InitialStartTimeDuringResize = Section.HasStartFrame() ? Section.GetInclusiveStartFrame() : 0;
 }
 
 void FGeometryCacheSection::ResizeSection(ESequencerSectionResizeMode ResizeMode, FFrameNumber ResizeTime)
@@ -256,28 +256,28 @@ void FGeometryCacheSection::SlipSection(FFrameNumber SlipTime)
 }
 
 
-FGeometryCacheTrackEditor::FGeometryCacheTrackEditor( TSharedRef<ISequencer> InSequencer )
-	: FMovieSceneTrackEditor( InSequencer ) 
+FGeometryCacheTrackEditor::FGeometryCacheTrackEditor(TSharedRef<ISequencer> InSequencer)
+	: FMovieSceneTrackEditor(InSequencer)
 { }
 
 
-TSharedRef<ISequencerTrackEditor> FGeometryCacheTrackEditor::CreateTrackEditor( TSharedRef<ISequencer> InSequencer )
+TSharedRef<ISequencerTrackEditor> FGeometryCacheTrackEditor::CreateTrackEditor(TSharedRef<ISequencer> InSequencer)
 {
-	return MakeShareable( new FGeometryCacheTrackEditor( InSequencer ) );
+	return MakeShareable(new FGeometryCacheTrackEditor(InSequencer));
 }
 
 
-bool FGeometryCacheTrackEditor::SupportsType( TSubclassOf<UMovieSceneTrack> Type ) const
+bool FGeometryCacheTrackEditor::SupportsType(TSubclassOf<UMovieSceneTrack> Type) const
 {
 	return Type == UMovieSceneGeometryCacheTrack::StaticClass();
 }
 
 
-TSharedRef<ISequencerSection> FGeometryCacheTrackEditor::MakeSectionInterface( UMovieSceneSection& SectionObject, UMovieSceneTrack& Track, FGuid ObjectBinding )
+TSharedRef<ISequencerSection> FGeometryCacheTrackEditor::MakeSectionInterface(UMovieSceneSection& SectionObject, UMovieSceneTrack& Track, FGuid ObjectBinding)
 {
-	check( SupportsType( SectionObject.GetOuter()->GetClass() ) );
-	
-	return MakeShareable( new FGeometryCacheSection(SectionObject, GetSequencer()) );
+	check(SupportsType(SectionObject.GetOuter()->GetClass()));
+
+	return MakeShareable(new FGeometryCacheSection(SectionObject, GetSequencer()));
 }
 
 void FGeometryCacheTrackEditor::BuildObjectBindingTrackMenu(FMenuBuilder& MenuBuilder, const FGuid& ObjectBinding, const UClass* ObjectClass)
@@ -317,11 +317,11 @@ void FGeometryCacheTrackEditor::BuildGeometryCacheTrack(FGuid ObjectBinding, UGe
 	}
 }
 
-FKeyPropertyResult FGeometryCacheTrackEditor::AddKeyInternal( FFrameNumber KeyTime, UObject* Object, UGeometryCacheComponent* GeomCacheComp, UMovieSceneTrack* Track)
+FKeyPropertyResult FGeometryCacheTrackEditor::AddKeyInternal(FFrameNumber KeyTime, UObject* Object, UGeometryCacheComponent* GeomCacheComp, UMovieSceneTrack* Track)
 {
 	FKeyPropertyResult KeyPropertyResult;
 
-	FFindOrCreateHandleResult HandleResult = FindOrCreateHandleToObject( Object );
+	FFindOrCreateHandleResult HandleResult = FindOrCreateHandleToObject(Object);
 	FGuid ObjectHandle = HandleResult.Handle;
 	KeyPropertyResult.bHandleCreated |= HandleResult.bWasCreated;
 	if (ObjectHandle.IsValid())
@@ -336,7 +336,7 @@ FKeyPropertyResult FGeometryCacheTrackEditor::AddKeyInternal( FFrameNumber KeyTi
 		{
 			Track->Modify();
 
-			UMovieSceneSection* NewSection = Cast<UMovieSceneGeometryCacheTrack>(Track)->AddNewAnimation( KeyTime, GeomCacheComp );
+			UMovieSceneSection* NewSection = Cast<UMovieSceneGeometryCacheTrack>(Track)->AddNewAnimation(KeyTime, GeomCacheComp);
 			KeyPropertyResult.bTrackModified = true;
 
 			GetSequencer()->EmptySelection();
@@ -362,7 +362,7 @@ TSharedPtr<SWidget> FGeometryCacheTrackEditor::BuildOutlinerEditWidget(const FGu
 			FMenuBuilder MenuBuilder(true, nullptr);
 
 			BuildGeometryCacheTrack(ObjectBinding, GeomMeshComp, Track);
-			
+
 			return MenuBuilder.MakeWidget();
 		};
 
@@ -378,7 +378,7 @@ TSharedPtr<SWidget> FGeometryCacheTrackEditor::BuildOutlinerEditWidget(const FGu
 	{
 		return TSharedPtr<SWidget>();
 	}
-	
+
 }
 
 const FSlateBrush* FGeometryCacheTrackEditor::GetIconBrush() const
diff --git a/Engine/Plugins/Experimental/GeometryCache/Source/GeometryCacheTracks/Classes/MovieSceneGeometryCacheSection.h b/Engine/Plugins/Experimental/GeometryCache/Source/GeometryCacheTracks/Classes/MovieSceneGeometryCacheSection.h
--- a/Engine/Plugins/Experimental/GeometryCache/Source/GeometryCacheTracks/Classes/MovieSceneGeometryCacheSection.h
+++ b/Engine/Plugins/Experimental/GeometryCache/Source/GeometryCacheTracks/Classes/MovieSceneGeometryCacheSection.h
@@ -18,42 +18,40 @@ struct FMovieSceneGeometryCacheParams
 	FMovieSceneGeometryCacheParams();
 
 	/** Gets the animation duration, modified by play rate */
-	float GetDuration() const { return FMath::IsNearlyZero(PlayRate) || GeometryCache.ResolveObject() == nullptr ? 0.f : Cast<UGeometryCacheComponent>(GeometryCache.ResolveObject())->GetDuration() / PlayRate; }
+	float GetDuration() const { return FMath::IsNearlyZero(PlayRate) || GeometryCache.Get() == nullptr ? 0.f : Cast<UGeometryCacheComponent>(GeometryCache.Get())->GetDuration() / PlayRate; }
 
 	/** Gets the animation sequence length, not modified by play rate */
-	float GetSequenceLength() const { return GeometryCache.ResolveObject() != nullptr ? Cast<UGeometryCacheComponent>(GeometryCache.ResolveObject())->GetDuration() : 0.f; }
-
-	/** The animation this section plays */
-	UPROPERTY(EditAnywhere, Category="GeometryCache", meta=(AllowedClasses = "GeometryCacheComponent"))
-	FSoftObjectPath GeometryCache;
+	float GetSequenceLength() const { return GeometryCache.Get() != nullptr ? Cast<UGeometryCacheComponent>(GeometryCache.Get())->GetDuration() : 0.f; }
 
 	/** The offset into the beginning of the animation clip */
-	UPROPERTY(EditAnywhere, Category="GeometryCache")
+	UPROPERTY(EditAnywhere, Category = "GeometryCache")
 	FFrameNumber StartFrameOffset;
-	
+
 	/** The offset into the end of the animation clip */
-	UPROPERTY(EditAnywhere, Category="GeometryCache")
+	UPROPERTY(EditAnywhere, Category = "GeometryCache")
 	FFrameNumber EndFrameOffset;
-	
+
 	/** The playback rate of the animation clip */
-	UPROPERTY(EditAnywhere, Category="GeometryCache")
+	UPROPERTY(EditAnywhere, Category = "GeometryCache")
 	float PlayRate;
 
 	/** Reverse the playback of the animation clip */
-	UPROPERTY(EditAnywhere, Category="Animation")
-	uint32 bReverse:1;
+	UPROPERTY(EditAnywhere, Category = "Animation")
+	uint32 bReverse : 1;
 
 	UPROPERTY()
 	float StartOffset_DEPRECATED;
-	
+
 	UPROPERTY()
 	float EndOffset_DEPRECATED;
+
+	TWeakObjectPtr<UGeometryCacheComponent> GeometryCache;
 };
 
 /**
  * Movie scene section that control geometry cache playback
  */
-UCLASS( MinimalAPI )
+UCLASS(MinimalAPI)
 class UMovieSceneGeometryCacheSection
 	: public UMovieSceneSection
 {
@@ -61,11 +59,11 @@ class UMovieSceneGeometryCacheSection
 
 public:
 
-	UPROPERTY(EditAnywhere, Category = "Animation", meta=(ShowOnlyInnerProperties))
+	UPROPERTY(EditAnywhere, Category = "Animation", meta = (ShowOnlyInnerProperties))
 	FMovieSceneGeometryCacheParams Params;
 
 	/** Get Frame Time as Animation Time*/
-    virtual float MapTimeToAnimation(FFrameTime InPosition, FFrameRate InFrameRate) const;
+	virtual float MapTimeToAnimation(FFrameTime InPosition, FFrameRate InFrameRate) const;
 
 protected:
 	//~ UMovieSceneSection interface
@@ -92,5 +90,4 @@ private:
 
 #endif
 
-
 };
diff --git a/Engine/Plugins/Experimental/GeometryCache/Source/GeometryCacheTracks/Private/MovieSceneGeometryCacheSection.cpp b/Engine/Plugins/Experimental/GeometryCache/Source/GeometryCacheTracks/Private/MovieSceneGeometryCacheSection.cpp
--- a/Engine/Plugins/Experimental/GeometryCache/Source/GeometryCacheTracks/Private/MovieSceneGeometryCacheSection.cpp
+++ b/Engine/Plugins/Experimental/GeometryCache/Source/GeometryCacheTracks/Private/MovieSceneGeometryCacheSection.cpp
@@ -22,8 +22,8 @@ FMovieSceneGeometryCacheParams::FMovieSceneGeometryCacheParams()
 	bReverse = false;
 }
 
-UMovieSceneGeometryCacheSection::UMovieSceneGeometryCacheSection( const FObjectInitializer& ObjectInitializer )
-	: Super( ObjectInitializer )
+UMovieSceneGeometryCacheSection::UMovieSceneGeometryCacheSection(const FObjectInitializer& ObjectInitializer)
+	: Super(ObjectInitializer)
 {
 	BlendType = EMovieSceneBlendType::Absolute;
 	EvalOptions.EnableAndSetCompletionMode(EMovieSceneCompletionMode::ProjectDefault);
@@ -89,10 +89,10 @@ FFrameNumber GetStartOffsetAtTrimTime(FQualifiedFrameTime TrimTime, const FMovie
 TOptional<TRange<FFrameNumber> > UMovieSceneGeometryCacheSection::GetAutoSizeRange() const
 {
 	FFrameRate FrameRate = GetTypedOuter<UMovieScene>()->GetTickResolution();
-
 	FFrameTime AnimationLength = Params.GetSequenceLength() * FrameRate;
+	int32 IFrameNumber = AnimationLength.FrameNumber.Value + (int)(AnimationLength.GetSubFrame() + 0.5f);
 
-	return TRange<FFrameNumber>(GetInclusiveStartFrame(), GetInclusiveStartFrame() + AnimationLength.FrameNumber);
+	return TRange<FFrameNumber>(GetInclusiveStartFrame(), GetInclusiveStartFrame() + IFrameNumber + 1);
 }
 
 
@@ -133,11 +133,11 @@ void UMovieSceneGeometryCacheSection::GetSnapTimes(TArray<FFrameNumber>& OutSnap
 {
 	Super::GetSnapTimes(OutSnapTimes, bGetSectionBorders);
 
-	const FFrameRate   FrameRate  = GetTypedOuter<UMovieScene>()->GetTickResolution();
+	const FFrameRate   FrameRate = GetTypedOuter<UMovieScene>()->GetTickResolution();
 	const FFrameNumber StartFrame = GetInclusiveStartFrame();
-	const FFrameNumber EndFrame   = GetExclusiveEndFrame() - 1; // -1 because we don't need to add the end frame twice
+	const FFrameNumber EndFrame = GetExclusiveEndFrame() - 1; // -1 because we don't need to add the end frame twice
 
-	const float AnimPlayRate     = FMath::IsNearlyZero(Params.PlayRate) ? 1.0f : Params.PlayRate;
+	const float AnimPlayRate = FMath::IsNearlyZero(Params.PlayRate) ? 1.0f : Params.PlayRate;
 	const float SeqLengthSeconds = Params.GetSequenceLength() - FrameRate.AsSeconds(Params.StartFrameOffset + Params.EndFrameOffset) / AnimPlayRate;
 
 	FFrameTime SequenceFrameLength = SeqLengthSeconds * FrameRate;
@@ -155,7 +155,7 @@ void UMovieSceneGeometryCacheSection::GetSnapTimes(TArray<FFrameNumber>& OutSnap
 
 float UMovieSceneGeometryCacheSection::MapTimeToAnimation(FFrameTime InPosition, FFrameRate InFrameRate) const
 {
-	FMovieSceneGeometryCacheSectionTemplateParameters TemplateParams(Params, GetInclusiveStartFrame(), GetExclusiveEndFrame());
+	FMovieSceneGeometryCacheSectionTemplateParameters TemplateParams(const_cast<FMovieSceneGeometryCacheParams&> (Params), GetInclusiveStartFrame(), GetExclusiveEndFrame());
 	return TemplateParams.MapTimeToAnimation(InPosition, InFrameRate);
 }
 
@@ -181,7 +181,7 @@ void UMovieSceneGeometryCacheSection::PostEditChangeProperty(FPropertyChangedEve
 		{
 			float CurrentDuration = MovieScene::DiscreteSize(GetRange());
 			float NewDuration = CurrentDuration * (PreviousPlayRate / NewPlayRate);
-			SetEndFrame( GetInclusiveStartFrame() + FMath::FloorToInt(NewDuration) );
+			SetEndFrame(GetInclusiveStartFrame() + FMath::FloorToInt(NewDuration));
 
 			PreviousPlayRate = NewPlayRate;
 		}
diff --git a/Engine/Plugins/Experimental/GeometryCache/Source/GeometryCacheTracks/Private/MovieSceneGeometryCacheTemplate.cpp b/Engine/Plugins/Experimental/GeometryCache/Source/GeometryCacheTracks/Private/MovieSceneGeometryCacheTemplate.cpp
--- a/Engine/Plugins/Experimental/GeometryCache/Source/GeometryCacheTracks/Private/MovieSceneGeometryCacheTemplate.cpp
+++ b/Engine/Plugins/Experimental/GeometryCache/Source/GeometryCacheTracks/Private/MovieSceneGeometryCacheTemplate.cpp
@@ -45,35 +45,67 @@ struct FPreAnimatedGeometryCacheTokenProducer : IMovieScenePreAnimatedTokenProdu
 /** A movie scene execution token that executes a geometry cache */
 struct FGeometryCacheExecutionToken
 	: IMovieSceneExecutionToken
-	
+
 {
-	FGeometryCacheExecutionToken(const FMovieSceneGeometryCacheSectionTemplateParameters &InParams):
+	FGeometryCacheExecutionToken(const FMovieSceneGeometryCacheSectionTemplateParameters &InParams) :
 		Params(InParams)
 	{}
 
+	static UGeometryCacheComponent* GeometryMeshComponentFromObject(UObject* BoundObject)
+	{
+		if (AActor* Actor = Cast<AActor>(BoundObject))
+		{
+			for (UActorComponent* Component : Actor->GetComponents())
+			{
+				if (UGeometryCacheComponent* GeometryMeshComp = Cast<UGeometryCacheComponent>(Component))
+				{
+					return GeometryMeshComp;
+				}
+			}
+		}
+		else if (UGeometryCacheComponent* GeometryMeshComp = Cast<UGeometryCacheComponent>(BoundObject))
+		{
+			if (GeometryMeshComp->GetGeometryCache())
+			{
+				return GeometryMeshComp;
+			}
+		}
+		return nullptr;
+	}
+
 	/** Execute this token, operating on all objects referenced by 'Operand' */
 	virtual void Execute(const FMovieSceneContext& Context, const FMovieSceneEvaluationOperand& Operand, FPersistentEvaluationData& PersistentData, IMovieScenePlayer& Player) override
 	{
-		if (Params.GeometryCache.ResolveObject())
+		MOVIESCENE_DETAILED_SCOPE_CYCLE_COUNTER(MovieSceneEval_GeometryCache_TokenExecute)
+		if (Operand.ObjectBindingID.IsValid())
 		{
-			MOVIESCENE_DETAILED_SCOPE_CYCLE_COUNTER(MovieSceneEval_GeometryCache_TokenExecute)
-
-			UGeometryCacheComponent* GeometryCache = Cast<UGeometryCacheComponent>(Params.GeometryCache.ResolveObject());
-
-			Player.SavePreAnimatedState(*GeometryCache, FPreAnimatedGeometryCacheTokenProducer::GetAnimTypeID(), FPreAnimatedGeometryCacheTokenProducer());
-			GeometryCache->SetManualTick(true);
-			// calculate the time at which to evaluate the animation
-			float EvalTime = Params.MapTimeToAnimation(Context.GetTime(), Context.GetFrameRate());
-			GeometryCache->TickAtThisTime(EvalTime, true, Params.bReverse, true);
+			for (TWeakObjectPtr<> WeakObj : Player.FindBoundObjects(Operand))
+			{
+				if (UObject* Obj = WeakObj.Get())
+				{
+					UGeometryCacheComponent* GeometryCache = GeometryMeshComponentFromObject(Obj);
+					if (GeometryCache)
+					{
+						if (Params.SectionParams->GeometryCache.Get() == nullptr)
+						{
+							Params.SectionParams->GeometryCache = GeometryCache;
+						}
+						Player.SavePreAnimatedState(*GeometryCache, FPreAnimatedGeometryCacheTokenProducer::GetAnimTypeID(), FPreAnimatedGeometryCacheTokenProducer());
+						GeometryCache->SetManualTick(true);
+						// calculate the time at which to evaluate the animation
+						float EvalTime = Params.MapTimeToAnimation(Context.GetTime(), Context.GetFrameRate());
+						GeometryCache->TickAtThisTime(EvalTime, true, Params.SectionParams->bReverse, true);
+					}
+				}
+			}
 		}
-
 	}
 
 	FMovieSceneGeometryCacheSectionTemplateParameters Params;
 };
 
 FMovieSceneGeometryCacheSectionTemplate::FMovieSceneGeometryCacheSectionTemplate(const UMovieSceneGeometryCacheSection& InSection)
-	: Params(InSection.Params, InSection.GetInclusiveStartFrame(), InSection.GetExclusiveEndFrame())
+	: Params(const_cast<FMovieSceneGeometryCacheParams&> (InSection.Params), InSection.GetInclusiveStartFrame(), InSection.GetExclusiveEndFrame())
 {
 }
 
@@ -82,28 +114,34 @@ FMovieSceneGeometryCacheSectionTemplate::FMovieSceneGeometryCacheSectionTemplate
 void FMovieSceneGeometryCacheSectionTemplate::Evaluate(const FMovieSceneEvaluationOperand& Operand, const FMovieSceneContext& Context, const FPersistentEvaluationData& PersistentData, FMovieSceneExecutionTokens& ExecutionTokens) const
 {
 	MOVIESCENE_DETAILED_SCOPE_CYCLE_COUNTER(MovieSceneEval_GeometryCache_Evaluate)
-	ExecutionTokens.Add(FGeometryCacheExecutionToken(Params));
+		ExecutionTokens.Add(FGeometryCacheExecutionToken(Params));
 }
 
 float FMovieSceneGeometryCacheSectionTemplateParameters::MapTimeToAnimation(FFrameTime InPosition, FFrameRate InFrameRate) const
 {
-	InPosition = FMath::Clamp(InPosition, FFrameTime(SectionStartTime), FFrameTime(SectionEndTime-1));
+	FFrameTime AnimationLength = SectionParams->GetSequenceLength() * InFrameRate;
+	int32 LengthInFrames = AnimationLength.FrameNumber.Value + (int)(AnimationLength.GetSubFrame() + 0.5f) + 1;
+	//we only play end if we are not looping, and assuming we are looping if Length is greater than default length;
+	bool bLooping = (SectionEndTime.Value - SectionStartTime.Value) > LengthInFrames;
+
+	InPosition = FMath::Clamp(InPosition, FFrameTime(SectionStartTime), FFrameTime(SectionEndTime - 1));
 
-	const float SectionPlayRate = PlayRate;
+	const float SectionPlayRate = SectionParams->PlayRate;
 	const float AnimPlayRate = FMath::IsNearlyZero(SectionPlayRate) ? 1.0f : SectionPlayRate;
 
-	const float SeqLength = GetSequenceLength() - InFrameRate.AsSeconds(StartFrameOffset + EndFrameOffset);
+	const float SeqLength = SectionParams->GetSequenceLength() - InFrameRate.AsSeconds(SectionParams->StartFrameOffset + SectionParams->EndFrameOffset);
 
 	float AnimPosition = FFrameTime::FromDecimal((InPosition - SectionStartTime).AsDecimal() * AnimPlayRate) / InFrameRate;
-	if (SeqLength > 0.f)
+	if (SeqLength > 0.f && (bLooping || !FMath::IsNearlyEqual(AnimPosition, SeqLength, 1e-4f)))
 	{
 		AnimPosition = FMath::Fmod(AnimPosition, SeqLength);
 	}
-	AnimPosition += InFrameRate.AsSeconds(StartFrameOffset);
-	if (bReverse)
+	AnimPosition += InFrameRate.AsSeconds(SectionParams->StartFrameOffset);
+	if (SectionParams->bReverse)
 	{
-		AnimPosition = (SeqLength - (AnimPosition - InFrameRate.AsSeconds(StartFrameOffset))) + InFrameRate.AsSeconds(StartFrameOffset);
+		AnimPosition = (SeqLength - (AnimPosition - InFrameRate.AsSeconds(SectionParams->StartFrameOffset))) + InFrameRate.AsSeconds(SectionParams->StartFrameOffset);
 	}
 
 	return AnimPosition;
-}
+
+}
\ No newline at end of file
diff --git a/Engine/Plugins/Experimental/GeometryCache/Source/GeometryCacheTracks/Private/MovieSceneGeometryCacheTemplate.h b/Engine/Plugins/Experimental/GeometryCache/Source/GeometryCacheTracks/Private/MovieSceneGeometryCacheTemplate.h
--- a/Engine/Plugins/Experimental/GeometryCache/Source/GeometryCacheTracks/Private/MovieSceneGeometryCacheTemplate.h
+++ b/Engine/Plugins/Experimental/GeometryCache/Source/GeometryCacheTracks/Private/MovieSceneGeometryCacheTemplate.h
@@ -9,19 +9,21 @@
 #include "MovieSceneGeometryCacheTemplate.generated.h"
 
 USTRUCT()
-struct FMovieSceneGeometryCacheSectionTemplateParameters : public FMovieSceneGeometryCacheParams
+struct FMovieSceneGeometryCacheSectionTemplateParameters
 {
 	GENERATED_BODY()
 
-	FMovieSceneGeometryCacheSectionTemplateParameters() {}
-	FMovieSceneGeometryCacheSectionTemplateParameters(const FMovieSceneGeometryCacheParams& BaseParams, FFrameNumber InSectionStartTime, FFrameNumber InSectionEndTime)
-		: FMovieSceneGeometryCacheParams(BaseParams)
+	FMovieSceneGeometryCacheSectionTemplateParameters() :SectionParams(nullptr) {}
+	FMovieSceneGeometryCacheSectionTemplateParameters(FMovieSceneGeometryCacheParams& BaseParams, FFrameNumber InSectionStartTime, FFrameNumber InSectionEndTime)
+		: SectionParams(&BaseParams)
 		, SectionStartTime(InSectionStartTime)
 		, SectionEndTime(InSectionEndTime)
 	{}
 
 	float MapTimeToAnimation(FFrameTime InPosition, FFrameRate InFrameRate) const;
 
+	FMovieSceneGeometryCacheParams* SectionParams;
+
 	UPROPERTY()
 	FFrameNumber SectionStartTime;
 
@@ -33,7 +35,7 @@ USTRUCT()
 struct FMovieSceneGeometryCacheSectionTemplate : public FMovieSceneEvalTemplate
 {
 	GENERATED_BODY()
-	
+
 	FMovieSceneGeometryCacheSectionTemplate() {}
 	FMovieSceneGeometryCacheSectionTemplate(const UMovieSceneGeometryCacheSection& Section);
 
