Moved range checks outside Alembic frame interpolation loops for more realistic development mode performance
Added stat scope for Alembic frame interpolation

diff --git a/Engine/Plugins/Experimental/GeometryCache/Source/GeometryCache/Private/GeometryCacheSceneProxy.cpp b/Engine/Plugins/Experimental/GeometryCache/Source/GeometryCache/Private/GeometryCacheSceneProxy.cpp
--- a/Engine/Plugins/Experimental/GeometryCache/Source/GeometryCache/Private/GeometryCacheSceneProxy.cpp
+++ b/Engine/Plugins/Experimental/GeometryCache/Source/GeometryCache/Private/GeometryCacheSceneProxy.cpp
@@ -21,6 +21,7 @@ DECLARE_DWORD_COUNTER_STAT(TEXT("Batch Count"), STAT_GeometryCacheSceneProxy_Mes
 DECLARE_CYCLE_STAT(TEXT("Vertex Buffer Update"), STAT_VertexBufferUpdate, STATGROUP_GeometryCache);
 DECLARE_CYCLE_STAT(TEXT("Index Buffer Update"), STAT_IndexBufferUpdate, STATGROUP_GeometryCache);
 DECLARE_CYCLE_STAT(TEXT("Buffer Update Task"), STAT_BufferUpdateTask, STATGROUP_GeometryCache);
+DECLARE_CYCLE_STAT(TEXT("InterpolateFrames"), STAT_InterpolateFrames, STATGROUP_GeometryCache);
 
 static TAutoConsoleVariable<int32> CVarOffloadUpdate(
 	TEXT("GeometryCache.OffloadUpdate"),
@@ -647,6 +648,7 @@ void FGeometryCacheSceneProxy::FrameUpdate() const
 			// Can we interpolate the vertex data?
 			if (bCanInterpolate && !bDecoderError && CVarInterpolateFrames.GetValueOnRenderThread() != 0)
 			{
+				SCOPE_CYCLE_COUNTER(STAT_InterpolateFrames);
 				// Interpolate if the time has changed.
 				// note: This is a bit precarious as this code is called multiple times per frame. This ensures
 				// we only interpolate once (which is a nice optimization) but more importantly that we only
@@ -682,44 +684,93 @@ void FGeometryCacheSceneProxy::FrameUpdate() const
 					const int32 InterpFixed = (int32)(InterpolationFactor * 255.0f);
 					const int32 OneMinusInterpFixed = 255 - InterpFixed;
 
-					for (int32 Index = 0; Index < NumVerts; ++Index)
 					{
-						const FVector& PositionA = TrackProxy->MeshData->Positions[Index];
-						const FVector& PositionB = TrackProxy->NextFrameMeshData->Positions[Index];
-						InterpolatedPositions[Index] = PositionA * OneMinusInterp + PositionB* InterpolationFactor;
+						check(TrackProxy->MeshData->Positions.Num() >= NumVerts);
+						check(TrackProxy->NextFrameMeshData->Positions.Num() >= NumVerts);
+						check(InterpolatedPositions.Num() >= NumVerts);
+						const FVector* PositionAPtr = TrackProxy->MeshData->Positions.GetData();
+						const FVector* PositionBPtr = TrackProxy->NextFrameMeshData->Positions.GetData();
+						FVector* InterpolationPositionsPtr = InterpolatedPositions.GetData();
+						for (int32 Index = 0; Index < NumVerts; ++Index)
+						{
+							const FVector& PositionA = PositionAPtr[Index];
+							const FVector& PositionB = PositionBPtr[Index];
+							InterpolationPositionsPtr[Index] = PositionA * OneMinusInterp + PositionB * InterpolationFactor;
+						}
 					}
-
-					for (int32 Index = 0; Index < NumVerts; ++Index)
+					
 					{
-						// The following are already 8 bit so quantized enough we can do exact equal comparisons
-						const FPackedNormal& TangentXA = TrackProxy->MeshData->TangentsX[Index];
-						const FPackedNormal& TangentXB = TrackProxy->NextFrameMeshData->TangentsX[Index];
-						const FPackedNormal& TangentZA = TrackProxy->MeshData->TangentsZ[Index];
-						const FPackedNormal& TangentZB = TrackProxy->NextFrameMeshData->TangentsZ[Index];
-
-						InterpolatedTangentX[Index] = InterpolatePackedNormal(TangentXA, TangentXB, InterpFixed, OneMinusInterpFixed);
-						InterpolatedTangentZ[Index] = InterpolatePackedNormal(TangentZA, TangentZB, InterpFixed, OneMinusInterpFixed);
+						check(TrackProxy->MeshData->TangentsX.Num() >= NumVerts);
+						check(TrackProxy->NextFrameMeshData->TangentsX.Num() >= NumVerts);
+						check(TrackProxy->MeshData->TangentsZ.Num() >= NumVerts);
+						check(TrackProxy->NextFrameMeshData->TangentsZ.Num() >= NumVerts);
+						check(InterpolatedTangentX.Num() >= NumVerts);
+						check(InterpolatedTangentZ.Num() >= NumVerts);
+						const FPackedNormal* TangentXAPtr = TrackProxy->MeshData->TangentsX.GetData();
+						const FPackedNormal* TangentXBPtr = TrackProxy->NextFrameMeshData->TangentsX.GetData();
+						const FPackedNormal* TangentZAPtr = TrackProxy->MeshData->TangentsZ.GetData();
+						const FPackedNormal* TangentZBPtr = TrackProxy->NextFrameMeshData->TangentsZ.GetData();
+						FPackedNormal* InterpolatedTangentXPtr = InterpolatedTangentX.GetData();
+						FPackedNormal* InterpolatedTangentZPtr = InterpolatedTangentZ.GetData();
+						for (int32 Index = 0; Index < NumVerts; ++Index)
+						{
+							// The following are already 8 bit so quantized enough we can do exact equal comparisons
+							const FPackedNormal& TangentXA = TangentXAPtr[Index];
+							const FPackedNormal& TangentXB = TangentXBPtr[Index];
+							const FPackedNormal& TangentZA = TangentZAPtr[Index];
+							const FPackedNormal& TangentZB = TangentZBPtr[Index];
+
+							InterpolatedTangentXPtr[Index] = InterpolatePackedNormal(TangentXA, TangentXB, InterpFixed, OneMinusInterpFixed);
+							InterpolatedTangentZPtr[Index] = InterpolatePackedNormal(TangentZA, TangentZB, InterpFixed, OneMinusInterpFixed);
+						}
 					}
+					
 
-					if (TrackProxy->MeshData->VertexInfo.bHasColor0) for (int32 Index = 0; Index < NumVerts; ++Index)
+					if (TrackProxy->MeshData->VertexInfo.bHasColor0)
 					{
-						const FColor& ColorA = TrackProxy->MeshData->Colors[Index];
-						const FColor& ColorB = TrackProxy->NextFrameMeshData->Colors[Index];
-						InterpolatedColors[Index] = InterpolatePackedColor(ColorA, ColorB, InterpFixed, OneMinusInterpFixed);
+						check(TrackProxy->MeshData->Colors.Num() >= NumVerts);
+						check(TrackProxy->NextFrameMeshData->Colors.Num() >= NumVerts);
+						check(InterpolatedColors.Num() >= NumVerts);
+						const FColor* ColorAPtr = TrackProxy->MeshData->Colors.GetData();
+						const FColor* ColorBPtr = TrackProxy->NextFrameMeshData->Colors.GetData();
+						FColor* InterpolatedColorsPtr = InterpolatedColors.GetData();
+						for (int32 Index = 0; Index < NumVerts; ++Index)
+						{
+							const FColor& ColorA = ColorAPtr[Index];
+							const FColor& ColorB = ColorBPtr[Index];
+							InterpolatedColorsPtr[Index] = InterpolatePackedColor(ColorA, ColorB, InterpFixed, OneMinusInterpFixed);
+						}
 					}
 
-					if (TrackProxy->MeshData->VertexInfo.bHasUV0) for (int32 Index = 0; Index < NumVerts; ++Index)
+					if (TrackProxy->MeshData->VertexInfo.bHasUV0)
 					{
-						const FVector2D& UVA = TrackProxy->MeshData->TextureCoordinates[Index];
-						const FVector2D& UVB = TrackProxy->NextFrameMeshData->TextureCoordinates[Index];
-
-						InterpolatedUVs[Index] = UVA * OneMinusInterp + UVB * InterpolationFactor;
+						check(TrackProxy->MeshData->TextureCoordinates.Num() >= NumVerts);
+						check(TrackProxy->NextFrameMeshData->TextureCoordinates.Num() >= NumVerts);
+						check(InterpolatedUVs.Num() >= NumVerts);
+						const FVector2D* UVAPtr = TrackProxy->MeshData->TextureCoordinates.GetData();
+						const FVector2D* UVBPtr = TrackProxy->NextFrameMeshData->TextureCoordinates.GetData();
+						FVector2D* InterpolatedUVsPtr = InterpolatedUVs.GetData();
+						for (int32 Index = 0; Index < NumVerts; ++Index)
+						{
+							const FVector2D& UVA = UVAPtr[Index];
+							const FVector2D& UVB = UVBPtr[Index];
+							InterpolatedUVsPtr[Index] = UVA * OneMinusInterp + UVB * InterpolationFactor;
+						}
 					}
 
-					if (bHasMotionVectors) for (int32 Index = 0; Index < NumVerts; ++Index)
+					if (bHasMotionVectors)
 					{
-						InterpolatedMotionVectors[Index] = TrackProxy->MeshData->MotionVectors[Index] * OneMinusInterp +
-							TrackProxy->NextFrameMeshData->MotionVectors[Index] * InterpolationFactor;
+						check(TrackProxy->MeshData->MotionVectors.Num() >= NumVerts);
+						check(TrackProxy->NextFrameMeshData->MotionVectors.Num() >= NumVerts);
+						check(InterpolatedMotionVectors.Num() >= NumVerts);
+						const FVector* MotionVectorAPtr = TrackProxy->MeshData->MotionVectors.GetData();
+						const FVector* MotionVectorBPtr = TrackProxy->NextFrameMeshData->MotionVectors.GetData();
+						FVector* InterpolatedMotionVectorsPtr = InterpolatedMotionVectors.GetData();
+
+						for (int32 Index = 0; Index < NumVerts; ++Index)
+						{
+							InterpolatedMotionVectorsPtr[Index] = MotionVectorAPtr[Index] * OneMinusInterp + MotionVectorBPtr[Index] * InterpolationFactor;
+						}
 					}
 
 					// Upload other non-motionblurred data
