Removed position encode/decoding dependency on index buffer as it is no longer needed because parallelogram prediction has been removed.
Refactored code so read functions now take an explicit FHuffmanBitStreamReader instead of implicitly using the one from the decoding context.
Split compressed data into chunks.
Parallel decode of chunks.

diff --git a/Engine/Plugins/Experimental/GeometryCache/Source/GeometryCache/Private/CodecV1.cpp b/Engine/Plugins/Experimental/GeometryCache/Source/GeometryCache/Private/CodecV1.cpp
--- a/Engine/Plugins/Experimental/GeometryCache/Source/GeometryCache/Private/CodecV1.cpp
+++ b/Engine/Plugins/Experimental/GeometryCache/Source/GeometryCache/Private/CodecV1.cpp
@@ -10,6 +10,12 @@
 #include "CodecV1Test.h"
 #include "Stats/StatsMisc.h"
 
+#define OLD_MAGIC 123
+#define CURRENT_MAGIC 124
+
+#define MAX_CHUNK_VERTICES	1024u
+#define CHUNK_MAX_TRIANGLES	2048u
+#define MAX_CHUNK_INDICES	(CHUNK_MAX_TRIANGLES * 3)
 /*
 //#define DEBUG_DUMP_FRAMES // Dump raw frames to file during encoding for debugging
 //#define DEBUG_CODECTEST_DUMPED_FRAMES // Test codec using dumped raw files
@@ -257,7 +263,7 @@ void FCodecV1Encoder::EncodeIndexStream(const uint32* Stream, uint64 ElementOffs
 	Stats = FStreamEncodingStatistics(ByteCounter.Read(), ElementCount * sizeof(uint32), Quality);
 }
 
-void FCodecV1Encoder::EncodePositionStream(const FVector* VertexStream, uint64 VertexElementOffset, uint32 VertexElementCount, const uint32* IndexStream, uint64 IndexElementOffset, uint32 IndexElementCount, FStreamEncodingStatistics& Stats)
+void FCodecV1Encoder::EncodePositionStream(const FVector* VertexStream, uint64 VertexElementOffset, uint32 VertexElementCount, FStreamEncodingStatistics& Stats)
 {
 	FBitstreamWriterByteCounter ByteCounter(EncodingContext.Writer);
 
@@ -277,55 +283,41 @@ void FCodecV1Encoder::EncodePositionStream(const FVector* VertexStream, uint64 V
 	Header.QuantizationPrecision = QuantizationPrecision;
 	Header.Translation = QuantizedTranslationToCenter;
 	WriteBytes((void*)&Header, sizeof(Header));
-
-	uint32 EncodedVertexCount = 0;
 	
 	FIntVector Prediction(0, 0, 0);	// Previously seen position
-	int64 MaxEncounteredIndex = -1;
 		
 	FQualityMetric QualityMetric;
 
-	const uint8* RawElementDataIndices = (const uint8*)IndexStream;
 	const uint8* RawElementDataVertices = (const uint8*)VertexStream;
 
 	// Walk over indices/triangles
-	for (uint32 IndexIdx = 0; IndexIdx < IndexElementCount; ++IndexIdx)
+	for (uint32 ElementIdx = 0; ElementIdx < VertexElementCount; ++ElementIdx)
 	{
-		const uint32 IndexValue = *(const uint32*)RawElementDataIndices;
-		RawElementDataIndices += IndexElementOffset;
-
-		if ((int64)IndexValue > MaxEncounteredIndex)
-		{
-			MaxEncounteredIndex = (int64)IndexValue;
-
-			// Code a newly encountered vertex
-			FVector VertexValue = *(FVector*)RawElementDataVertices;
-			RawElementDataVertices += VertexElementOffset;
+		// Code a newly encountered vertex
+		FVector VertexValue = *(FVector*)RawElementDataVertices;
+		RawElementDataVertices += VertexElementOffset;
 
-			// Quantize
-			FIntVector Encoded = Quantizer.Quantize(VertexValue);
+		// Quantize
+		FIntVector Encoded = Quantizer.Quantize(VertexValue);
 
-			// Translate to center
-			FIntVector EncodedCentered = Encoded - QuantizedTranslationToCenter;
+		// Translate to center
+		FIntVector EncodedCentered = Encoded - QuantizedTranslationToCenter;
 
-			// Residual to code
-			FIntVector Residual = EncodedCentered - Prediction;
+		// Residual to code
+		FIntVector Residual = EncodedCentered - Prediction;
 				
-			// Write residual
-			WriteInt32(EncodingContext.ResidualVertexPosTable, Residual.X);
-			WriteInt32(EncodingContext.ResidualVertexPosTable, Residual.Y);
-			WriteInt32(EncodingContext.ResidualVertexPosTable, Residual.Z);
+		// Write residual
+		WriteInt32(EncodingContext.ResidualVertexPosTable, Residual.X);
+		WriteInt32(EncodingContext.ResidualVertexPosTable, Residual.Y);
+		WriteInt32(EncodingContext.ResidualVertexPosTable, Residual.Z);
 
-			EncodedVertexCount++;
-				
-			// Store previous encountered values
-			FIntVector Reconstructed = Prediction + Residual;
+		// Store previous encountered values
+		FIntVector Reconstructed = Prediction + Residual;
 
-			// Calculate error
-			FVector DequantReconstructed = Quantizer.Dequantize(Reconstructed);
-			QualityMetric.Register(VertexValue, DequantReconstructed);
-			Prediction = Reconstructed;
-		}
+		// Calculate error
+		FVector DequantReconstructed = Quantizer.Dequantize(Reconstructed);
+		QualityMetric.Register(VertexValue, DequantReconstructed);
+		Prediction = Reconstructed;
 	}
 
 	// Gather rate and quality statistics
@@ -595,6 +587,8 @@ bool FCodecV1Encoder::EncodeFrameData(FMemoryWriter& Writer, const FGeometryCach
 	const TArray<uint32>& Indices = MeshData.Indices;
 	const TArray<FVector>& MotionVectors = MeshData.MotionVectors;	
 
+	const FGeometryCacheVertexInfo& VertexInfo = MeshData.VertexInfo;
+
 	{
 		// Check if indices are referenced in order, i.e. if a previously unreferenced vertex is referenced by the index 
 		// list it's id will always be the next unreferenced id instead of some random unused id. E.g., ok: 1, 2, 3, 2, 4, not ok: 1, 2, 4
@@ -607,35 +601,58 @@ bool FCodecV1Encoder::EncodeFrameData(FMemoryWriter& Writer, const FGeometryCach
 			MaxIndex = FMath::Max(Indices[i], MaxIndex);
 		}
 	}
+
+	checkf(!VertexInfo.bHasMotionVectors || MotionVectors.Num() > 0, TEXT("No motion vectors while VertexInfo states otherwise"));
+
 	FBitstreamWriterByteCounter TotalByteCounter(EncodingContext.Writer);
-	const FGeometryCacheVertexInfo& VertexInfo = MeshData.VertexInfo;
+	
+	const uint32 NumIndices = MeshData.Indices.Num();
+	const uint32 NumVertices = MeshData.Positions.Num();
 
-	// Encode streams	
-	if (!VertexInfo.bConstantIndices)
-	{
-		EncodeIndexStream(&Indices[0], Indices.GetTypeSize(), Indices.Num(), Statistics.Indices);
-	}
-	EncodePositionStream(&Positions[0], Positions.GetTypeSize(), Positions.Num(), &Indices[0], Indices.GetTypeSize(), Indices.Num(), Statistics.Vertices);
-	if (VertexInfo.bHasColor0)
-	{
-		EncodeColorStream(&Colors[0], Colors.GetTypeSize(), Colors.Num(), Statistics.Colors);
-	}
-	if (VertexInfo.bHasTangentX)
-	{
-		EncodeNormalStream(&TangentsX[0], TangentsX.GetTypeSize(), TangentsX.Num(), EncodingContext.ResidualNormalTangentXTable, Statistics.TangentX);
-	}
-	if (VertexInfo.bHasTangentZ)
-	{
-		EncodeNormalStream(&TangentsZ[0], TangentsZ.GetTypeSize(), TangentsZ.Num(), EncodingContext.ResidualNormalTangentZTable, Statistics.TangentY);
-	}
-	if (VertexInfo.bHasUV0)
-	{
-		EncodeUVStream(&TextureCoordinates[0], TextureCoordinates.GetTypeSize(), TextureCoordinates.Num(), Statistics.TexCoords);
-	}
-	if (VertexInfo.bHasMotionVectors)
+	const uint32 NumChunks = FMath::Max((NumIndices + MAX_CHUNK_INDICES - 1) / MAX_CHUNK_INDICES, (NumVertices + MAX_CHUNK_VERTICES - 1) / MAX_CHUNK_VERTICES);
+	TArray<uint32> ChunkOffsets;
+	ChunkOffsets.AddUninitialized(NumChunks);
+
+	// Encode streams
+	for (uint32 ChunkIndex = 0; ChunkIndex < NumChunks; ChunkIndex++)
 	{
-		checkf(MotionVectors.Num() > 0, TEXT("No motion vectors while VertexInfo states otherwise"));
-		EncodeMotionVectorStream(&MotionVectors[0], MotionVectors.GetTypeSize(), MotionVectors.Num(), Statistics.MotionVectors);		
+		uint32 IndexOffset = ChunkIndex * (NumIndices / NumChunks);
+		uint32 VertexOffset = ChunkIndex * (NumVertices / NumChunks);
+		uint32 NumChunkIndices = (ChunkIndex == NumChunks - 1) ? (NumIndices - IndexOffset) : (NumIndices / NumChunks);
+		uint32 NumChunkVertices = (ChunkIndex == NumChunks - 1) ? (NumVertices - VertexOffset) : (NumVertices / NumChunks);
+
+		EncodingContext.Writer->Align();
+		ChunkOffsets[ChunkIndex] = EncodingContext.Writer->GetNumBytes();
+
+		if (NumChunkIndices > 0 && !VertexInfo.bConstantIndices)
+		{
+			EncodeIndexStream(&Indices[IndexOffset], Indices.GetTypeSize(), NumChunkIndices, Statistics.Indices);
+		}
+
+		if (NumChunkVertices > 0)
+		{
+			EncodePositionStream(&Positions[VertexOffset], Positions.GetTypeSize(), NumChunkVertices, Statistics.Vertices);
+			if (VertexInfo.bHasColor0)
+			{
+				EncodeColorStream(&Colors[VertexOffset], Colors.GetTypeSize(), NumChunkVertices, Statistics.Colors);
+			}
+			if (VertexInfo.bHasTangentX)
+			{
+				EncodeNormalStream(&TangentsX[VertexOffset], TangentsX.GetTypeSize(), NumChunkVertices, EncodingContext.ResidualNormalTangentXTable, Statistics.TangentX);
+			}
+			if (VertexInfo.bHasTangentZ)
+			{
+				EncodeNormalStream(&TangentsZ[VertexOffset], TangentsZ.GetTypeSize(), NumChunkVertices, EncodingContext.ResidualNormalTangentZTable, Statistics.TangentY);
+			}
+			if (VertexInfo.bHasUV0)
+			{
+				EncodeUVStream(&TextureCoordinates[VertexOffset], TextureCoordinates.GetTypeSize(), NumChunkVertices, Statistics.TexCoords);
+			}
+			if (VertexInfo.bHasMotionVectors)
+			{
+				EncodeMotionVectorStream(&MotionVectors[VertexOffset], MotionVectors.GetTypeSize(), NumChunkVertices, Statistics.MotionVectors);
+			}
+		}
 	}
 	
 	BitWriter.Close();
@@ -644,7 +661,7 @@ bool FCodecV1Encoder::EncodeFrameData(FMemoryWriter& Writer, const FGeometryCach
 	{
 		// Write out bitstream
 		FCodedFrameHeader Header = { 0 };
-		Header.Magic = 123;
+		Header.Magic = CURRENT_MAGIC;
 		Header.VertexCount = (uint32)Positions.Num();
 		Header.IndexCount = (uint32)Indices.Num();
 		uint32 PayloadSize = BitWriter.GetNumBytes();
@@ -652,6 +669,7 @@ bool FCodecV1Encoder::EncodeFrameData(FMemoryWriter& Writer, const FGeometryCach
 		Writer.Serialize(&Header, sizeof(Header));	// Write header				
 		Writer << MeshData.BatchesInfo;	// Uncompressed data: bounding box & material list
 		Writer << MeshData.BoundingBox;
+		Writer << ChunkOffsets;
 		Writer.Serialize((void*)BitWriter.GetBytes().GetData(), PayloadSize);	// Write payload
 	}	
 
@@ -825,7 +843,7 @@ FCodecV1Decoder::FCodecV1Decoder()
 	}	
 }
 
-void FCodecV1Decoder::DecodeIndexStream(uint32* Stream, uint64 ElementOffset, uint32 ElementCount)
+void FCodecV1Decoder::DecodeIndexStream(FHuffmanBitStreamReader& Reader, uint32* Stream, uint64 ElementOffset, uint32 ElementCount)
 {
 	const uint8* RawElementData = (const uint8*)Stream;
 
@@ -833,7 +851,7 @@ void FCodecV1Decoder::DecodeIndexStream(uint32* Stream, uint64 ElementOffset, ui
 	for (uint32 ElementIdx = 0; ElementIdx < ElementCount; ++ElementIdx, RawElementData += ElementOffset)
 	{
 		// Read coded residual
-		int32 DecodedResidual = ReadInt32(DecodingContext.ResidualIndicesTable);
+		int32 DecodedResidual = ReadInt32(Reader, DecodingContext.ResidualIndicesTable);
 		Value += DecodedResidual;
 
 		// Save result to our list
@@ -841,11 +859,11 @@ void FCodecV1Decoder::DecodeIndexStream(uint32* Stream, uint64 ElementOffset, ui
 	}
 }
 
-void FCodecV1Decoder::DecodeMotionVectorStream(FVector* Stream, uint64 ElementOffset, uint32 ElementCount)
+void FCodecV1Decoder::DecodeMotionVectorStream(FHuffmanBitStreamReader& Reader, FVector* Stream, uint64 ElementOffset, uint32 ElementCount)
 {
 	// Read header
 	FMotionVectorStreamHeader Header;
-	ReadBytes((void*)&Header, sizeof(Header));
+	ReadBytes(Reader, (void*)&Header, sizeof(Header));
 
 	FQuantizerVector3 Quantizer(Header.QuantizationPrecision); // We quantize MVs to a certain precision just like the positions
 
@@ -857,20 +875,20 @@ void FCodecV1Decoder::DecodeMotionVectorStream(FVector* Stream, uint64 ElementOf
 	{
 		// Read coded residual
 		FIntVector DecodedResidual;
-		DecodedResidual.X = ReadInt32(DecodingContext.ResidualMotionVectorTable);
-		DecodedResidual.Y = ReadInt32(DecodingContext.ResidualMotionVectorTable);
-		DecodedResidual.Z = ReadInt32(DecodingContext.ResidualMotionVectorTable);
+		DecodedResidual.X = ReadInt32(Reader, DecodingContext.ResidualMotionVectorTable);
+		DecodedResidual.Y = ReadInt32(Reader, DecodingContext.ResidualMotionVectorTable);
+		DecodedResidual.Z = ReadInt32(Reader, DecodingContext.ResidualMotionVectorTable);
 
 		QuantizedValue += DecodedResidual;
 		*(FVector*)RawElementData = Quantizer.Dequantize(QuantizedValue);
 	}
 }
 
-void FCodecV1Decoder::DecodeUVStream(FVector2D* Stream, uint64 ElementOffset, uint32 ElementCount)
+void FCodecV1Decoder::DecodeUVStream(FHuffmanBitStreamReader& Reader, FVector2D* Stream, uint64 ElementOffset, uint32 ElementCount)
 {
 	// Read header
 	FUVStreamHeader Header;
-	ReadBytes((void*)&Header, sizeof(Header));
+	ReadBytes(Reader, (void*)&Header, sizeof(Header));
 
 	FQuantizerVector2 Quantizer(Header.Range, Header.QuantizationBits); // We quantize UVs to a number of bits, set in the bitstream header
 	
@@ -882,15 +900,15 @@ void FCodecV1Decoder::DecodeUVStream(FVector2D* Stream, uint64 ElementOffset, ui
 	{
 		// Read coded residual
 		FIntVector DecodedResidual;
-		DecodedResidual.X = ReadInt32(DecodingContext.ResidualUVTable);
-		DecodedResidual.Y = ReadInt32(DecodingContext.ResidualUVTable);
+		DecodedResidual.X = ReadInt32(Reader, DecodingContext.ResidualUVTable);
+		DecodedResidual.Y = ReadInt32(Reader, DecodingContext.ResidualUVTable);
 
 		QuantizedValue += DecodedResidual;
 		*(FVector2D*)RawElementData = Quantizer.Dequantize(QuantizedValue);
 	}
 }
 
-void FCodecV1Decoder::DecodeNormalStream(FPackedNormal* Stream, uint64 ElementOffset, uint32 ElementCount, FHuffmanDecodeTable& Table)
+void FCodecV1Decoder::DecodeNormalStream(FHuffmanBitStreamReader& Reader, FPackedNormal* Stream, uint64 ElementOffset, uint32 ElementCount, FHuffmanDecodeTable& Table)
 {
 	uint8 x = 128, y = 128, z = 128, w = 128;
 
@@ -898,7 +916,6 @@ void FCodecV1Decoder::DecodeNormalStream(FPackedNormal* Stream, uint64 ElementOf
 
 	check(HUFFMAN_MAX_CODE_LENGTH * 4 <= MINIMUM_BITS_AFTER_REFILL);	// Make sure we can safely decode all 4 symbols with a single refill
 
-	FHuffmanBitStreamReader& Reader = *DecodingContext.Reader;
 	for (uint32 ElementIdx = 0; ElementIdx < ElementCount; ++ElementIdx, RawElementData += ElementOffset) // Walk normals
 	{
 		// Read coded residual
@@ -916,7 +933,7 @@ void FCodecV1Decoder::DecodeNormalStream(FPackedNormal* Stream, uint64 ElementOf
 	}
 }
 
-void FCodecV1Decoder::DecodeColorStream(FColor* Stream, uint64 ElementOffset, uint32 ElementCount)
+void FCodecV1Decoder::DecodeColorStream(FHuffmanBitStreamReader& Reader, FColor* Stream, uint64 ElementOffset, uint32 ElementCount)
 {
 	FIntVector4 QuantizedValue(128, 128, 128, 255);
 
@@ -926,15 +943,15 @@ void FCodecV1Decoder::DecodeColorStream(FColor* Stream, uint64 ElementOffset, ui
 	{
 		FIntVector4 DecodedResidual(0, 0, 0, 0);
 
-		int32 SkipBit = ReadBits(1); // 1: Perfect prediction, nothing coded, 0: we have coded residuals
+		int32 SkipBit = ReadBits(Reader, 1); // 1: Perfect prediction, nothing coded, 0: we have coded residuals
 
 		if (SkipBit != 1) 
 		{
 			// Prediction not perfect, residual were coded
-			int32 DecodedResidualR = ReadInt32(DecodingContext.ResidualColorTable);
-			int32 DecodedResidualG = ReadInt32(DecodingContext.ResidualColorTable);
-			int32 DecodedResidualB = ReadInt32(DecodingContext.ResidualColorTable);
-			int32 DecodedResidualA = ReadInt32(DecodingContext.ResidualColorTable);
+			int32 DecodedResidualR = ReadInt32(Reader, DecodingContext.ResidualColorTable);
+			int32 DecodedResidualG = ReadInt32(Reader, DecodingContext.ResidualColorTable);
+			int32 DecodedResidualB = ReadInt32(Reader, DecodingContext.ResidualColorTable);
+			int32 DecodedResidualA = ReadInt32(Reader, DecodingContext.ResidualColorTable);
 
 			DecodedResidual = FIntVector4(DecodedResidualR, DecodedResidualG, DecodedResidualB, DecodedResidualA);
 			QuantizedValue = FCodecV1SharedTools::SumVector4(QuantizedValue, DecodedResidual);
@@ -948,53 +965,39 @@ void FCodecV1Decoder::DecodeColorStream(FColor* Stream, uint64 ElementOffset, ui
 	}
 }
 
-void FCodecV1Decoder::DecodePositionStream(const uint32* IndexStream, uint64 IndexElementOffset, uint32 IndexElementCount, FVector* VertexStream, uint64 VertexElementOffset, uint32 MaxVertexElementCount)
+void FCodecV1Decoder::DecodePositionStream(FHuffmanBitStreamReader& Reader, FVector* VertexStream, uint64 VertexElementOffset, uint32 VertexElementCount)
 {
-	checkf(IndexElementCount > 0, TEXT("You cannot decode vertex stream before the index stream was decoded"));
-
 	// Read header
 	FVertexStreamHeader Header;
-	ReadBytes((void*)&Header, sizeof(Header));
+	ReadBytes(Reader, (void*)&Header, sizeof(Header));
 
 	FQuantizerVector3 Quantizer(Header.QuantizationPrecision);
 	
-	int64 MaxEncounteredIndex = -1; // We rely on indices being references in order, a requirement of the encoder and enforced by the preprocessor
 	uint32 DecodedVertexCount = 0;
 
-	const uint8* RawElementDataIndices = (const uint8*)IndexStream;
 	const uint8* RawElementDataVertices = (const uint8*)VertexStream;
 
 	FIntVector QuantizedValue(0, 0, 0);
 
 	// Walk over indices/triangles
-	for (uint32 IndexIdx = 0; IndexIdx < IndexElementCount; ++IndexIdx)
+	for (uint32 ElementIdx = 0; ElementIdx < VertexElementCount; ++ElementIdx)
 	{
-		const uint32 IndexValue = *(const uint32*)RawElementDataIndices;
-		RawElementDataIndices += IndexElementOffset;
-
-		if ((int64)IndexValue > MaxEncounteredIndex)
-		{
-			MaxEncounteredIndex = (int64)IndexValue;
-			checkf(DecodedVertexCount < MaxVertexElementCount, TEXT("Encountering more vertices than we have encoded. Encoding and decoding algorithms don't seem to match. Please make sure the preprocessor has processed the mesh such that vertexes are referenced in-order, i.e. if a previously unreferenced vertex is referenced by the index list it's id will always be the next unreferenced id instead of some random unused id."));
-
-			// Read coded residual
-			const FIntVector DecodedResidual = { ReadInt32(DecodingContext.ResidualVertexPosTable), ReadInt32(DecodingContext.ResidualVertexPosTable), ReadInt32(DecodingContext.ResidualVertexPosTable) };
-			DecodedVertexCount++;
+		// Read coded residual
+		const FIntVector DecodedResidual = { ReadInt32(Reader, DecodingContext.ResidualVertexPosTable), ReadInt32(Reader, DecodingContext.ResidualVertexPosTable), ReadInt32(Reader, DecodingContext.ResidualVertexPosTable) };
+		DecodedVertexCount++;
 
-			QuantizedValue += DecodedResidual;
+		QuantizedValue += DecodedResidual;
 
-			// Save result to our list
-			FVector* Value = (FVector*)RawElementDataVertices;
-			*Value = Quantizer.Dequantize(QuantizedValue + Header.Translation);
-			RawElementDataVertices += VertexElementOffset;		
-		}
+		// Save result to our list
+		FVector* Value = (FVector*)RawElementDataVertices;
+		*Value = Quantizer.Dequantize(QuantizedValue + Header.Translation);
+		RawElementDataVertices += VertexElementOffset;
 	}
 }
 
-void FCodecV1Decoder::SetupAndReadTables()
+void FCodecV1Decoder::SetupAndReadTables(FHuffmanBitStreamReader& Reader)
 {
 	// Initialize and read Huffman tables from the bitstream
-	FHuffmanBitStreamReader& Reader = *DecodingContext.Reader;	
 	const FGeometryCacheVertexInfo& VertexInfo = DecodingContext.MeshData->VertexInfo;
 		
 	if (!VertexInfo.bConstantIndices)
@@ -1025,22 +1028,23 @@ void FCodecV1Decoder::SetupAndReadTables()
 	// Add additional tables here
 }
 
-void FCodecV1Decoder::ReadCodedStreamDescription()
+void FCodecV1Decoder::ReadCodedStreamDescription(FHuffmanBitStreamReader& Reader)
 {	
 	FGeometryCacheVertexInfo& VertexInfo = DecodingContext.MeshData->VertexInfo;
 	
 	VertexInfo = FGeometryCacheVertexInfo();
-	VertexInfo.bHasTangentX = (ReadBits(1) == 1);
-	VertexInfo.bHasTangentZ = (ReadBits(1) == 1);
-	VertexInfo.bHasUV0 = (ReadBits(1) == 1);
-	VertexInfo.bHasColor0 = (ReadBits(1) == 1);
-	VertexInfo.bHasMotionVectors = (ReadBits(1) == 1);
-
-	VertexInfo.bConstantUV0 = (ReadBits(1) == 1);
-	VertexInfo.bConstantColor0 = (ReadBits(1) == 1);
-	VertexInfo.bConstantIndices = (ReadBits(1) == 1);
+	VertexInfo.bHasTangentX = (ReadBits(Reader, 1) == 1);
+	VertexInfo.bHasTangentZ = (ReadBits(Reader, 1) == 1);
+	VertexInfo.bHasUV0 = (ReadBits(Reader, 1) == 1);
+	VertexInfo.bHasColor0 = (ReadBits(Reader, 1) == 1);
+	VertexInfo.bHasMotionVectors = (ReadBits(Reader, 1) == 1);
+
+	VertexInfo.bConstantUV0 = (ReadBits(Reader, 1) == 1);
+	VertexInfo.bConstantColor0 = (ReadBits(Reader, 1) == 1);
+	VertexInfo.bConstantIndices = (ReadBits(Reader, 1) == 1);
 }
 
+DECLARE_CYCLE_STAT(TEXT("FCodecV1Decoder"), STAT_CodecV1Decoder, STATGROUP_GeometryCache);
 DECLARE_CYCLE_STAT(TEXT("SetupAndReadTables"), STAT_SetupAndReadTables, STATGROUP_GeometryCache);
 DECLARE_CYCLE_STAT(TEXT("DecodeIndexStream"), STAT_DecodeIndexStream, STATGROUP_GeometryCache);
 DECLARE_CYCLE_STAT(TEXT("DecodePositionStream"), STAT_DecodePositionStream, STATGROUP_GeometryCache);
@@ -1048,6 +1052,7 @@ DECLARE_CYCLE_STAT(TEXT("DecodeColorStream"), STAT_DecodeColorStream, STATGROUP_
 DECLARE_CYCLE_STAT(TEXT("DecodeTangentXStream"), STAT_DecodeTangentXStream, STATGROUP_GeometryCache);
 DECLARE_CYCLE_STAT(TEXT("DecodeTangentZStream"), STAT_DecodeTangentZStream, STATGROUP_GeometryCache);
 DECLARE_CYCLE_STAT(TEXT("DecodeUVStream"), STAT_DecodeUVStream, STATGROUP_GeometryCache);
+DECLARE_CYCLE_STAT(TEXT("DecodeMotionVectorStream"), STAT_DecodeMotionVectorStream, STATGROUP_GeometryCache);
 
 
 bool FCodecV1Decoder::DecodeFrameData(FBufferReader& Reader, FGeometryCacheMeshData &OutMeshData)
@@ -1058,7 +1063,7 @@ bool FCodecV1Decoder::DecodeFrameData(FBufferReader& Reader, FGeometryCacheMeshD
 	FCodedFrameHeader Header;
 	Reader.Serialize(&Header, sizeof(Header));
 
-	if (Header.Magic != 123)
+	if (Header.Magic != OLD_MAGIC && Header.Magic != CURRENT_MAGIC)
 	{
 		UE_LOG(LogGeoCaStreamingCodecV1, Error, TEXT("Incompatible bitstream found"));
 		return false;
@@ -1068,6 +1073,26 @@ bool FCodecV1Decoder::DecodeFrameData(FBufferReader& Reader, FGeometryCacheMeshD
 	Reader << OutMeshData.BatchesInfo; 
 	Reader << OutMeshData.BoundingBox;
 
+	const uint32 NumIndices = Header.IndexCount;
+	const uint32 NumVertices = Header.VertexCount;
+
+	uint32 NumChunks;
+	TArray<uint32> ChunkOffsets;
+
+	bool IsChunked = false;
+	if (Header.Magic == OLD_MAGIC)
+	{
+		// Compatibility with old files. Just treat everything as one big chunk
+		IsChunked = false;
+		NumChunks = 1;
+	}
+	else
+	{
+		IsChunked = true;
+		Reader << ChunkOffsets;
+		NumChunks = ChunkOffsets.Num();
+	}
+	
 	DecodingContext = { 0 };
 	DecodingContext.MeshData = &OutMeshData;	
 	
@@ -1076,108 +1101,179 @@ bool FCodecV1Decoder::DecodeFrameData(FBufferReader& Reader, FGeometryCacheMeshD
 	Bytes.AddUninitialized(Header.PayloadSize + 16);	// Overallocate by 16 bytes to ensure BitReader can safely perform uint64 reads.
 	Reader.Serialize(Bytes.GetData(), Header.PayloadSize);
 	FHuffmanBitStreamReader BitReader(Bytes.GetData(), Bytes.Num());
-	DecodingContext.Reader = &BitReader;
-
+	
 	// Read which vertex attributes are in the bit stream
-	ReadCodedStreamDescription();
+	ReadCodedStreamDescription(BitReader);
 
 	// Restore entropy coding contexts
 	{
 		SCOPE_CYCLE_COUNTER(STAT_SetupAndReadTables);
-		SetupAndReadTables();
+		SetupAndReadTables(BitReader);
 	}
 	
-	
+	// Allocate buffers
 	{
-		// Decode streams
-		
-		TArray<uint32>& Indices = OutMeshData.Indices;
 		if (!OutMeshData.VertexInfo.bConstantIndices)
 		{
 			OutMeshData.Indices.Empty(Header.IndexCount);
 			OutMeshData.Indices.AddUninitialized(Header.IndexCount);
-			SCOPE_CYCLE_COUNTER(STAT_DecodeIndexStream);
-			DecodeIndexStream(&Indices[0], Indices.GetTypeSize(), Header.IndexCount);
 		}
-		
-		TArray<FVector>& Positions = OutMeshData.Positions;
+
 		OutMeshData.Positions.Empty(Header.VertexCount);
-		OutMeshData.Positions.AddZeroed(Header.VertexCount);
-		{
-			SCOPE_CYCLE_COUNTER(STAT_DecodePositionStream);
-			DecodePositionStream(&Indices[0], Indices.GetTypeSize(), Indices.Num(), &Positions[0], Positions.GetTypeSize(), Header.VertexCount);
-		}
+		OutMeshData.Positions.AddUninitialized(Header.VertexCount);
 
 		OutMeshData.Colors.Empty(Header.VertexCount);
-		OutMeshData.Colors.AddZeroed(Header.VertexCount);
 		if (OutMeshData.VertexInfo.bHasColor0)
 		{
-			TArray<FColor>& Colors = OutMeshData.Colors;
-			SCOPE_CYCLE_COUNTER(STAT_DecodeColorStream);
-			DecodeColorStream(&Colors[0], Colors.GetTypeSize(), Header.VertexCount);
+			OutMeshData.Colors.AddUninitialized(Header.VertexCount);
 		}
-
+		else
+		{
+			OutMeshData.Colors.AddZeroed(Header.VertexCount);
+		}
+		
+		
 		OutMeshData.TangentsX.Empty(Header.VertexCount);
-		OutMeshData.TangentsX.AddZeroed(Header.VertexCount);
 		if (OutMeshData.VertexInfo.bHasTangentX)
 		{
-			TArray<FPackedNormal>& TangentsX = OutMeshData.TangentsX;
-			SCOPE_CYCLE_COUNTER(STAT_DecodeTangentXStream);
-			DecodeNormalStream(&TangentsX[0], TangentsX.GetTypeSize(), Header.VertexCount, DecodingContext.ResidualNormalTangentXTable);
+			OutMeshData.TangentsX.AddUninitialized(Header.VertexCount);
 		}
-
+		else
+		{
+			OutMeshData.TangentsX.AddZeroed(Header.VertexCount);
+		}
+		
 		OutMeshData.TangentsZ.Empty(Header.VertexCount);
-		OutMeshData.TangentsZ.AddZeroed(Header.VertexCount);
 		if (OutMeshData.VertexInfo.bHasTangentZ)
 		{
-			TArray<FPackedNormal>& TangentsZ = OutMeshData.TangentsZ;
-			SCOPE_CYCLE_COUNTER(STAT_DecodeTangentZStream);
-			DecodeNormalStream(&TangentsZ[0], TangentsZ.GetTypeSize(), Header.VertexCount, DecodingContext.ResidualNormalTangentZTable);
+			OutMeshData.TangentsZ.AddUninitialized(Header.VertexCount);
 		}
-
+		else
+		{
+			OutMeshData.TangentsZ.AddZeroed(Header.VertexCount);
+		}
+		
 		OutMeshData.TextureCoordinates.Empty(Header.VertexCount);
-		OutMeshData.TextureCoordinates.AddZeroed(Header.VertexCount);
 		if (OutMeshData.VertexInfo.bHasUV0)
 		{
-			TArray<FVector2D>& TextureCoordinates = OutMeshData.TextureCoordinates;
-			SCOPE_CYCLE_COUNTER(STAT_DecodeUVStream);
-			DecodeUVStream(&TextureCoordinates[0], TextureCoordinates.GetTypeSize(), Header.VertexCount);
+			OutMeshData.TextureCoordinates.AddUninitialized(Header.VertexCount);
 		}
-
-		TArray<FVector>& MotionVectors = OutMeshData.MotionVectors;
+		else
+		{
+			OutMeshData.TextureCoordinates.AddZeroed(Header.VertexCount);
+		}
+		
+		
 		OutMeshData.MotionVectors.Empty(Header.VertexCount);
 		if (OutMeshData.VertexInfo.bHasMotionVectors)
 		{
 			OutMeshData.MotionVectors.AddUninitialized(Header.VertexCount);
-			DecodeMotionVectorStream(&MotionVectors[0], MotionVectors.GetTypeSize(), Header.VertexCount);
 		}
-
-		if (CVarCodecDebug.GetValueOnAnyThread() == 1)
+		else
 		{
-			const float TimeFloat = DecodingTime.Get();
-			UE_LOG(LogGeoCaStreamingCodecV1, Log, TEXT("Decoded frame with %i vertices in %.2f milliseconds."), Positions.Num(), TimeFloat);
+			OutMeshData.MotionVectors.AddZeroed(Header.VertexCount);
 		}
 	}
-	DecodingContext.Reader = NULL;
+	
+	// Decompress chunks using the task graph. Note: ParallelFor is slow to wake threads. It can take up to 1ms!
+	// ParallelFor(NumChunks, [this, NumChunks, NumIndices, NumVertices, IsChunked, &Bytes, &BitReader, &ChunkOffsets, &OutMeshData](int32 ChunkIndex)
+	FGraphEventArray Handles;
+	Handles.AddDefaulted(NumChunks);
+	for (uint32 ChunkIndex = 0; ChunkIndex < NumChunks; ChunkIndex++)
+	{
+		Handles[ChunkIndex] = FFunctionGraphTask::CreateAndDispatchWhenReady([this, ChunkIndex, NumChunks, NumIndices, NumVertices, IsChunked, &Bytes, &BitReader, &ChunkOffsets, &OutMeshData]()
+		{
+			SCOPE_CYCLE_COUNTER(STAT_CodecV1Decoder);
+
+			uint32 IndexOffset = 0;
+			uint32 VertexOffset = 0;
+			uint32 NumChunkIndices = NumIndices;
+			uint32 NumChunkVertices = NumVertices;
+
+			uint32 ChunkOffset = 0;
+			if (IsChunked)
+			{
+				IndexOffset = (uint32)ChunkIndex * (NumIndices / NumChunks);
+				VertexOffset = (uint32)ChunkIndex * (NumVertices / NumChunks);
+				NumChunkIndices = (ChunkIndex + 1 == NumChunks) ? (NumIndices - IndexOffset) : (NumIndices / NumChunks);
+				NumChunkVertices = (ChunkIndex + 1 == NumChunks) ? (NumVertices - VertexOffset) : (NumVertices / NumChunks);
+				ChunkOffset = ChunkOffsets[ChunkIndex];
+			}
+
+			FHuffmanBitStreamReader ReaderObj(Bytes.GetData() + ChunkOffset, Bytes.Num() - ChunkOffset);
+			FHuffmanBitStreamReader& Reader = IsChunked ? ReaderObj : BitReader;
+
+			if (NumChunkIndices > 0 && !OutMeshData.VertexInfo.bConstantIndices)
+			{
+				SCOPE_CYCLE_COUNTER(STAT_DecodeIndexStream);
+				DecodeIndexStream(Reader, &OutMeshData.Indices[IndexOffset], OutMeshData.Indices.GetTypeSize(), NumChunkIndices);
+			}
+
+			if (NumChunkVertices > 0)
+			{
+				{
+					SCOPE_CYCLE_COUNTER(STAT_DecodePositionStream);
+					DecodePositionStream(Reader, &OutMeshData.Positions[VertexOffset], OutMeshData.Positions.GetTypeSize(), NumChunkVertices);
+				}
+
+				if (OutMeshData.VertexInfo.bHasColor0)
+				{
+					SCOPE_CYCLE_COUNTER(STAT_DecodeColorStream);
+					DecodeColorStream(Reader, &OutMeshData.Colors[VertexOffset], OutMeshData.Colors.GetTypeSize(), NumChunkVertices);
+				}
+
+				if (OutMeshData.VertexInfo.bHasTangentX)
+				{
+					SCOPE_CYCLE_COUNTER(STAT_DecodeTangentXStream);
+					DecodeNormalStream(Reader, &OutMeshData.TangentsX[VertexOffset], OutMeshData.TangentsX.GetTypeSize(), NumChunkVertices, DecodingContext.ResidualNormalTangentXTable);
+				}
+
+				if (OutMeshData.VertexInfo.bHasTangentZ)
+				{
+					SCOPE_CYCLE_COUNTER(STAT_DecodeTangentZStream);
+					DecodeNormalStream(Reader, &OutMeshData.TangentsZ[VertexOffset], OutMeshData.TangentsZ.GetTypeSize(), NumChunkVertices, DecodingContext.ResidualNormalTangentZTable);
+				}
+
+				if (OutMeshData.VertexInfo.bHasUV0)
+				{
+					SCOPE_CYCLE_COUNTER(STAT_DecodeUVStream);
+					DecodeUVStream(Reader, &OutMeshData.TextureCoordinates[VertexOffset], OutMeshData.TextureCoordinates.GetTypeSize(), NumChunkVertices);
+				}
+
+				if (OutMeshData.VertexInfo.bHasMotionVectors)
+				{
+					SCOPE_CYCLE_COUNTER(STAT_DecodeMotionVectorStream);
+					DecodeMotionVectorStream(Reader, &OutMeshData.MotionVectors[VertexOffset], OutMeshData.MotionVectors.GetTypeSize(), NumChunkVertices);
+				}
+			}
+		}, GET_STATID(STAT_CodecV1Decoder));
+	}
+	FTaskGraphInterface::Get().WaitUntilTasksComplete(Handles);	//TODO: can we sync later somehow?
 
+	if (CVarCodecDebug.GetValueOnAnyThread() == 1)
+	{
+		const float TimeFloat = DecodingTime.Get();
+		UE_LOG(LogGeoCaStreamingCodecV1, Log, TEXT("Decoded frame with %i vertices in %.2f milliseconds."), NumVertices, TimeFloat);
+	}
+	
 	return true;
 }
 
-void FCodecV1Decoder::ReadBytes(void* Data, uint32 NumBytes)
+void FCodecV1Decoder::ReadBytes(FHuffmanBitStreamReader& Reader, void* Data, uint32 NumBytes)
 {
 	uint8* ByteData = (uint8*)Data;
 
 	for (int64 ByteIndex = 0; ByteIndex < NumBytes; ++ByteIndex)
 	{
-		const uint32 ByteValue = DecodingContext.Reader->Read(8);
+		const uint32 ByteValue = Reader.Read(8);
 		*ByteData++ = ByteValue & 255;
 	}
 }
 
-int32 FCodecV1Decoder::ReadInt32(FHuffmanDecodeTable& ValueTable)
+int32 FCodecV1Decoder::ReadInt32(FHuffmanBitStreamReader& Reader, FHuffmanDecodeTable& ValueTable)
 {
 	// See write WriteInt32 for encoding details.
-	const int32 Packed = ReadSymbol(ValueTable);
+	const int32 Packed = ReadSymbol(Reader, ValueTable);
 	if (Packed < 4)
 	{
 		// [-2, 1] coded directly with no additional raw bits
@@ -1187,6 +1283,6 @@ int32 FCodecV1Decoder::ReadInt32(FHuffmanDecodeTable& ValueTable)
 	{
 		// At least one raw bit.
 		int32 NumRawBits = (Packed - 2) >> 1;
-		return ReadBitsNoRefill(NumRawBits) + HighBitsLUT[Packed];
+		return ReadBitsNoRefill(Reader, NumRawBits) + HighBitsLUT[Packed];
 	}
 }
\ No newline at end of file
diff --git a/Engine/Plugins/Experimental/GeometryCache/Source/GeometryCache/Private/CodecV1.h b/Engine/Plugins/Experimental/GeometryCache/Source/GeometryCache/Private/CodecV1.h
--- a/Engine/Plugins/Experimental/GeometryCache/Source/GeometryCache/Private/CodecV1.h
+++ b/Engine/Plugins/Experimental/GeometryCache/Source/GeometryCache/Private/CodecV1.h
@@ -84,8 +84,6 @@ struct FCodecV1DecodingContext
 {
 	/** Target mesh to decode to */
 	FGeometryCacheMeshData* MeshData;	
-	/** Reader to read bit stream from */
-	FHuffmanBitStreamReader* Reader;
 	/** Huffman tables */
 	FHuffmanDecodeTable ResidualIndicesTable;	
 	FHuffmanDecodeTable ResidualVertexPosTable;
@@ -205,7 +203,7 @@ private:
 	/** Encode a buffer with vertex indices */
 	void EncodeIndexStream(const uint32* Stream, uint64 ElementOffsetBytes, uint32 ElementCount, FStreamEncodingStatistics& Stats);
 	/** Encode a buffer with vertex positions */
-	void EncodePositionStream(const FVector* VertexStream, uint64 VertexElementOffset, uint32 VertexElementCount, const uint32* IndexStream, uint64 IndexElementOffset, uint32 IndexElementCount, FStreamEncodingStatistics& Stats);
+	void EncodePositionStream(const FVector* VertexStream, uint64 VertexElementOffset, uint32 VertexElementCount, FStreamEncodingStatistics& Stats);
 	/** Encode a buffer with vertex colors */
 	void EncodeColorStream(const FColor* Stream, uint64 ElementOffsetBytes, uint32 ElementCount, FStreamEncodingStatistics& Stats);
 	/** Encode a buffer with vertex normals */
@@ -300,50 +298,50 @@ public:
 	virtual bool DecodeFrameData(FBufferReader &Reader, FGeometryCacheMeshData &OutMeshData);
 private:
 
-	/** Decode a buffer of vertex indices from the bitstream */
-	void DecodeIndexStream(uint32* Stream, uint64 ElementOffset, uint32 ElementCount);
-	/** Decode a buffer of vertex positions from the bitstream */
-	void DecodePositionStream(const uint32* IndexStream, uint64 IndexElementOffset, uint32 IndexElementCount, FVector* VertexStream, uint64 VertexElementOffset, uint32 MaxVertexElementCount);
-	/** Decode a buffer of vertex colors from the bitstream */
-	void DecodeColorStream(FColor* Stream, uint64 ElementOffset, uint32 ElementCount);
-	/** Decode a buffer of vertex normals from the bitstream */
-	void DecodeNormalStream(FPackedNormal* Stream, uint64 ElementOffset, uint32 ElementCount, FHuffmanDecodeTable& Table);
-	/** Decode a buffer of vertex texture coordinates from the bitstream */
-	void DecodeUVStream(FVector2D* Stream, uint64 ElementOffset, uint32 ElementCount);
-	/** Decode a buffer of vertex motion vectors from the bitstream */
-	void DecodeMotionVectorStream(FVector* Stream, uint64 ElementOffset, uint32 ElementCount);
+	/** Decode a buffer of vertex indices from abitstream */
+	void DecodeIndexStream(FHuffmanBitStreamReader& Reader, uint32* Stream, uint64 ElementOffset, uint32 ElementCount);
+	/** Decode a buffer of vertex positions from a bitstream */
+	void DecodePositionStream(FHuffmanBitStreamReader& Reader, FVector* VertexStream, uint64 VertexElementOffset, uint32 VertexElementCount);
+	/** Decode a buffer of vertex colors from a bitstream */
+	void DecodeColorStream(FHuffmanBitStreamReader& Reader, FColor* Stream, uint64 ElementOffset, uint32 ElementCount);
+	/** Decode a buffer of vertex normals from a bitstream */
+	void DecodeNormalStream(FHuffmanBitStreamReader& Reader, FPackedNormal* Stream, uint64 ElementOffset, uint32 ElementCount, FHuffmanDecodeTable& Table);
+	/** Decode a buffer of vertex texture coordinates from a bitstream */
+	void DecodeUVStream(FHuffmanBitStreamReader& Reader, FVector2D* Stream, uint64 ElementOffset, uint32 ElementCount);
+	/** Decode a buffer of vertex motion vectors from a bitstream */
+	void DecodeMotionVectorStream(FHuffmanBitStreamReader& Reader, FVector* Stream, uint64 ElementOffset, uint32 ElementCount);
 	
-	/** Initialize and read Huffman tables from the bitstream */
-	void SetupAndReadTables();	
+	/** Initialize and read Huffman tables from a bitstream */
+	void SetupAndReadTables(FHuffmanBitStreamReader& Reader);
 	
-	/** Read a given number of bytes from the bit stream using the Huffman bitreader */
-	FORCEINLINE void ReadBytes(void* Data, uint32 NumBytes);
+	/** Read a given number of bytes from a bit stream using a Huffman bitreader */
+	FORCEINLINE void ReadBytes(FHuffmanBitStreamReader& Reader, void* Data, uint32 NumBytes);
 
-	/**	Read a 32-bit signed value from the bit stream using a Huffman coder */
-	FORCEINLINE int32 ReadInt32(FHuffmanDecodeTable& ValueTable);
+	/**	Read a 32-bit signed value from a bit stream using a Huffman coder */
+	FORCEINLINE int32 ReadInt32(FHuffmanBitStreamReader& Reader, FHuffmanDecodeTable& ValueTable);
 
-	/** Read a symbol from the bit stream using a Huffman Table */
-	FORCEINLINE int32 ReadSymbol(FHuffmanDecodeTable& table)
+	/** Read a symbol from a bit stream using a Huffman Table */
+	FORCEINLINE int32 ReadSymbol(FHuffmanBitStreamReader& Reader, FHuffmanDecodeTable& table)
 	{
-		return table.Decode(*DecodingContext.Reader);
+		return table.Decode(Reader);
 	}
 
-	/** Read a given number of bits from the bit stream */
-	FORCEINLINE int32 ReadBits(uint32 NumBits)
+	/** Read a given number of bits from a bit stream */
+	FORCEINLINE int32 ReadBits(FHuffmanBitStreamReader& Reader, uint32 NumBits)
 	{
-		return DecodingContext.Reader->Read(NumBits);
+		return Reader.Read(NumBits);
 	}
 
-	/** Read a given number of bits from the bit stream without refilling bits */
-	FORCEINLINE int32 ReadBitsNoRefill(uint32 NumBits)
+	/** Read a given number of bits from a bit stream without refilling bits */
+	FORCEINLINE int32 ReadBitsNoRefill(FHuffmanBitStreamReader& Reader, uint32 NumBits)
 	{
-		return DecodingContext.Reader->ReadNoRefill(NumBits);
+		return Reader.ReadNoRefill(NumBits);
 	}
 
-	/** Read info on the available streams from the bit stream */
-	void ReadCodedStreamDescription();
+	/** Read info on the available streams from a bit stream */
+	void ReadCodedStreamDescription(FHuffmanBitStreamReader& Reader);
 
-	/** Any context information to decode a frame in a sequence of frames, such as the bit steam and any Huffman tables used.*/
+	/** Any context information to decode a frame in a sequence of frames, such as any Huffman tables used. */
 	FCodecV1DecodingContext DecodingContext;	
 
 	int32 HighBitsLUT[64];
diff --git a/Engine/Plugins/Experimental/GeometryCache/Source/GeometryCache/Private/HuffmanBitStream.h b/Engine/Plugins/Experimental/GeometryCache/Source/GeometryCache/Private/HuffmanBitStream.h
--- a/Engine/Plugins/Experimental/GeometryCache/Source/GeometryCache/Private/HuffmanBitStream.h
+++ b/Engine/Plugins/Experimental/GeometryCache/Source/GeometryCache/Private/HuffmanBitStream.h
@@ -56,18 +56,26 @@ public:
 	}
 
 	/**
-	Close the stream. This ensures the stream's underlying byte buffer is correctly
-	flushed with all bits written to the stream.
-	Once closed you cannot longer Write additional data on the stream.
+	Flush pending bits and align to byte.
 	*/
-	void Close()
+	void Align()
 	{
 		// Round up to next byte by appending 0s.
-		if(BitBufferBits)
+		if (BitBufferBits)
 		{
 			Write(0, 8 - BitBufferBits);
 		}
 		check(BitBufferBits == 0);
+	}
+
+	/**
+	Close the stream. This ensures the stream's underlying byte buffer is correctly
+	flushed with all bits written to the stream.
+	Once closed you cannot longer Write additional data on the stream.
+	*/
+	void Close()
+	{
+		Align();
 		
 		bFlushed = true;
 	}
