Alembic decode: Use ParallelFor instead of individual tasks to get lower frame variance when there is competition for the worker threads.

diff --git a/Engine/Plugins/Experimental/GeometryCache/Source/GeometryCache/Private/CodecV1.cpp b/Engine/Plugins/Experimental/GeometryCache/Source/GeometryCache/Private/CodecV1.cpp
--- a/Engine/Plugins/Experimental/GeometryCache/Source/GeometryCache/Private/CodecV1.cpp
+++ b/Engine/Plugins/Experimental/GeometryCache/Source/GeometryCache/Private/CodecV1.cpp
@@ -9,6 +9,7 @@
 #include "HuffmanBitStream.h"
 #include "CodecV1Test.h"
 #include "Stats/StatsMisc.h"
+#include "Async/ParallelFor.h"
 
 #define OLD_MAGIC 123
 #define CURRENT_MAGIC 124
@@ -1176,79 +1177,73 @@ bool FCodecV1Decoder::DecodeFrameData(FBufferReader& Reader, FGeometryCacheMeshD
 	}
 	
 	// Decompress chunks using the task graph. Note: ParallelFor is slow to wake threads. It can take up to 1ms!
-	// ParallelFor(NumChunks, [this, NumChunks, NumIndices, NumVertices, IsChunked, &Bytes, &BitReader, &ChunkOffsets, &OutMeshData](int32 ChunkIndex)
-	FGraphEventArray Handles;
-	Handles.AddDefaulted(NumChunks);
-	for (uint32 ChunkIndex = 0; ChunkIndex < NumChunks; ChunkIndex++)
+	// TODO: can we sync later to get more overlap?
+	ParallelFor(NumChunks, [this, NumChunks, NumIndices, NumVertices, IsChunked, &Bytes, &BitReader, &ChunkOffsets, &OutMeshData](int32 ChunkIndex)
 	{
-		Handles[ChunkIndex] = FFunctionGraphTask::CreateAndDispatchWhenReady([this, ChunkIndex, NumChunks, NumIndices, NumVertices, IsChunked, &Bytes, &BitReader, &ChunkOffsets, &OutMeshData]()
+		SCOPE_CYCLE_COUNTER(STAT_CodecV1Decoder);
+
+		uint32 IndexOffset = 0;
+		uint32 VertexOffset = 0;
+		uint32 NumChunkIndices = NumIndices;
+		uint32 NumChunkVertices = NumVertices;
+
+		uint32 ChunkOffset = 0;
+		if (IsChunked)
 		{
-			SCOPE_CYCLE_COUNTER(STAT_CodecV1Decoder);
+			IndexOffset = (uint32)ChunkIndex * (NumIndices / NumChunks);
+			VertexOffset = (uint32)ChunkIndex * (NumVertices / NumChunks);
+			NumChunkIndices = (ChunkIndex + 1 == NumChunks) ? (NumIndices - IndexOffset) : (NumIndices / NumChunks);
+			NumChunkVertices = (ChunkIndex + 1 == NumChunks) ? (NumVertices - VertexOffset) : (NumVertices / NumChunks);
+			ChunkOffset = ChunkOffsets[ChunkIndex];
+		}
+
+		FHuffmanBitStreamReader ReaderObj(Bytes.GetData() + ChunkOffset, Bytes.Num() - ChunkOffset);
+		FHuffmanBitStreamReader& Reader = IsChunked ? ReaderObj : BitReader;
 
-			uint32 IndexOffset = 0;
-			uint32 VertexOffset = 0;
-			uint32 NumChunkIndices = NumIndices;
-			uint32 NumChunkVertices = NumVertices;
+		if (NumChunkIndices > 0 && !OutMeshData.VertexInfo.bConstantIndices)
+		{
+			SCOPE_CYCLE_COUNTER(STAT_DecodeIndexStream);
+			DecodeIndexStream(Reader, &OutMeshData.Indices[IndexOffset], OutMeshData.Indices.GetTypeSize(), NumChunkIndices);
+		}
 
-			uint32 ChunkOffset = 0;
-			if (IsChunked)
+		if (NumChunkVertices > 0)
+		{
 			{
-				IndexOffset = (uint32)ChunkIndex * (NumIndices / NumChunks);
-				VertexOffset = (uint32)ChunkIndex * (NumVertices / NumChunks);
-				NumChunkIndices = (ChunkIndex + 1 == NumChunks) ? (NumIndices - IndexOffset) : (NumIndices / NumChunks);
-				NumChunkVertices = (ChunkIndex + 1 == NumChunks) ? (NumVertices - VertexOffset) : (NumVertices / NumChunks);
-				ChunkOffset = ChunkOffsets[ChunkIndex];
+				SCOPE_CYCLE_COUNTER(STAT_DecodePositionStream);
+				DecodePositionStream(Reader, &OutMeshData.Positions[VertexOffset], OutMeshData.Positions.GetTypeSize(), NumChunkVertices);
 			}
 
-			FHuffmanBitStreamReader ReaderObj(Bytes.GetData() + ChunkOffset, Bytes.Num() - ChunkOffset);
-			FHuffmanBitStreamReader& Reader = IsChunked ? ReaderObj : BitReader;
+			if (OutMeshData.VertexInfo.bHasColor0)
+			{
+				SCOPE_CYCLE_COUNTER(STAT_DecodeColorStream);
+				DecodeColorStream(Reader, &OutMeshData.Colors[VertexOffset], OutMeshData.Colors.GetTypeSize(), NumChunkVertices);
+			}
 
-			if (NumChunkIndices > 0 && !OutMeshData.VertexInfo.bConstantIndices)
+			if (OutMeshData.VertexInfo.bHasTangentX)
 			{
-				SCOPE_CYCLE_COUNTER(STAT_DecodeIndexStream);
-				DecodeIndexStream(Reader, &OutMeshData.Indices[IndexOffset], OutMeshData.Indices.GetTypeSize(), NumChunkIndices);
+				SCOPE_CYCLE_COUNTER(STAT_DecodeTangentXStream);
+				DecodeNormalStream(Reader, &OutMeshData.TangentsX[VertexOffset], OutMeshData.TangentsX.GetTypeSize(), NumChunkVertices, DecodingContext.ResidualNormalTangentXTable);
 			}
 
-			if (NumChunkVertices > 0)
+			if (OutMeshData.VertexInfo.bHasTangentZ)
 			{
-				{
-					SCOPE_CYCLE_COUNTER(STAT_DecodePositionStream);
-					DecodePositionStream(Reader, &OutMeshData.Positions[VertexOffset], OutMeshData.Positions.GetTypeSize(), NumChunkVertices);
-				}
-
-				if (OutMeshData.VertexInfo.bHasColor0)
-				{
-					SCOPE_CYCLE_COUNTER(STAT_DecodeColorStream);
-					DecodeColorStream(Reader, &OutMeshData.Colors[VertexOffset], OutMeshData.Colors.GetTypeSize(), NumChunkVertices);
-				}
-
-				if (OutMeshData.VertexInfo.bHasTangentX)
-				{
-					SCOPE_CYCLE_COUNTER(STAT_DecodeTangentXStream);
-					DecodeNormalStream(Reader, &OutMeshData.TangentsX[VertexOffset], OutMeshData.TangentsX.GetTypeSize(), NumChunkVertices, DecodingContext.ResidualNormalTangentXTable);
-				}
-
-				if (OutMeshData.VertexInfo.bHasTangentZ)
-				{
-					SCOPE_CYCLE_COUNTER(STAT_DecodeTangentZStream);
-					DecodeNormalStream(Reader, &OutMeshData.TangentsZ[VertexOffset], OutMeshData.TangentsZ.GetTypeSize(), NumChunkVertices, DecodingContext.ResidualNormalTangentZTable);
-				}
-
-				if (OutMeshData.VertexInfo.bHasUV0)
-				{
-					SCOPE_CYCLE_COUNTER(STAT_DecodeUVStream);
-					DecodeUVStream(Reader, &OutMeshData.TextureCoordinates[VertexOffset], OutMeshData.TextureCoordinates.GetTypeSize(), NumChunkVertices);
-				}
-
-				if (OutMeshData.VertexInfo.bHasMotionVectors)
-				{
-					SCOPE_CYCLE_COUNTER(STAT_DecodeMotionVectorStream);
-					DecodeMotionVectorStream(Reader, &OutMeshData.MotionVectors[VertexOffset], OutMeshData.MotionVectors.GetTypeSize(), NumChunkVertices);
-				}
+				SCOPE_CYCLE_COUNTER(STAT_DecodeTangentZStream);
+				DecodeNormalStream(Reader, &OutMeshData.TangentsZ[VertexOffset], OutMeshData.TangentsZ.GetTypeSize(), NumChunkVertices, DecodingContext.ResidualNormalTangentZTable);
 			}
-		}, GET_STATID(STAT_CodecV1Decoder));
-	}
-	FTaskGraphInterface::Get().WaitUntilTasksComplete(Handles);	//TODO: can we sync later somehow?
+
+			if (OutMeshData.VertexInfo.bHasUV0)
+			{
+				SCOPE_CYCLE_COUNTER(STAT_DecodeUVStream);
+				DecodeUVStream(Reader, &OutMeshData.TextureCoordinates[VertexOffset], OutMeshData.TextureCoordinates.GetTypeSize(), NumChunkVertices);
+			}
+
+			if (OutMeshData.VertexInfo.bHasMotionVectors)
+			{
+				SCOPE_CYCLE_COUNTER(STAT_DecodeMotionVectorStream);
+				DecodeMotionVectorStream(Reader, &OutMeshData.MotionVectors[VertexOffset], OutMeshData.MotionVectors.GetTypeSize(), NumChunkVertices);
+			}
+		}
+	});
 
 	if (CVarCodecDebug.GetValueOnAnyThread() == 1)
 	{
