Fix for alembic crash

diff --git a/Engine/Plugins/Experimental/GeometryCache/Source/GeometryCache/Private/GeometryCacheComponent.cpp b/Engine/Plugins/Experimental/GeometryCache/Source/GeometryCache/Private/GeometryCacheComponent.cpp
--- a/Engine/Plugins/Experimental/GeometryCache/Source/GeometryCache/Private/GeometryCacheComponent.cpp
+++ b/Engine/Plugins/Experimental/GeometryCache/Source/GeometryCache/Private/GeometryCacheComponent.cpp
@@ -217,6 +217,7 @@ void UGeometryCacheComponent::UpdateLocalBounds()
 
 FPrimitiveSceneProxy* UGeometryCacheComponent::CreateSceneProxy()
 {
+	IGeometryCacheStreamingManager::Get().PrefetchData(this);
 	return new FGeometryCacheSceneProxy(this);
 }
 
