Fixed compiler warning on Android.

diff --git a/Engine/Plugins/Experimental/GeometryCache/Source/GeometryCache/Private/CodecV1.cpp b/Engine/Plugins/Experimental/GeometryCache/Source/GeometryCache/Private/CodecV1.cpp
--- a/Engine/Plugins/Experimental/GeometryCache/Source/GeometryCache/Private/CodecV1.cpp
+++ b/Engine/Plugins/Experimental/GeometryCache/Source/GeometryCache/Private/CodecV1.cpp
@@ -1198,49 +1198,49 @@ bool FCodecV1Decoder::DecodeFrameData(FBufferReader& Reader, FGeometryCacheMeshD
 		}
 
 		FHuffmanBitStreamReader ReaderObj(Bytes.GetData() + ChunkOffset, Bytes.Num() - ChunkOffset);
-		FHuffmanBitStreamReader& Reader = IsChunked ? ReaderObj : BitReader;
+		FHuffmanBitStreamReader& ChunkReader = IsChunked ? ReaderObj : BitReader;
 
 		if (NumChunkIndices > 0 && !OutMeshData.VertexInfo.bConstantIndices)
 		{
 			SCOPE_CYCLE_COUNTER(STAT_DecodeIndexStream);
-			DecodeIndexStream(Reader, &OutMeshData.Indices[IndexOffset], OutMeshData.Indices.GetTypeSize(), NumChunkIndices);
+			DecodeIndexStream(ChunkReader, &OutMeshData.Indices[IndexOffset], OutMeshData.Indices.GetTypeSize(), NumChunkIndices);
 		}
 
 		if (NumChunkVertices > 0)
 		{
 			{
 				SCOPE_CYCLE_COUNTER(STAT_DecodePositionStream);
-				DecodePositionStream(Reader, &OutMeshData.Positions[VertexOffset], OutMeshData.Positions.GetTypeSize(), NumChunkVertices);
+				DecodePositionStream(ChunkReader, &OutMeshData.Positions[VertexOffset], OutMeshData.Positions.GetTypeSize(), NumChunkVertices);
 			}
 
 			if (OutMeshData.VertexInfo.bHasColor0)
 			{
 				SCOPE_CYCLE_COUNTER(STAT_DecodeColorStream);
-				DecodeColorStream(Reader, &OutMeshData.Colors[VertexOffset], OutMeshData.Colors.GetTypeSize(), NumChunkVertices);
+				DecodeColorStream(ChunkReader, &OutMeshData.Colors[VertexOffset], OutMeshData.Colors.GetTypeSize(), NumChunkVertices);
 			}
 
 			if (OutMeshData.VertexInfo.bHasTangentX)
 			{
 				SCOPE_CYCLE_COUNTER(STAT_DecodeTangentXStream);
-				DecodeNormalStream(Reader, &OutMeshData.TangentsX[VertexOffset], OutMeshData.TangentsX.GetTypeSize(), NumChunkVertices, DecodingContext.ResidualNormalTangentXTable);
+				DecodeNormalStream(ChunkReader, &OutMeshData.TangentsX[VertexOffset], OutMeshData.TangentsX.GetTypeSize(), NumChunkVertices, DecodingContext.ResidualNormalTangentXTable);
 			}
 
 			if (OutMeshData.VertexInfo.bHasTangentZ)
 			{
 				SCOPE_CYCLE_COUNTER(STAT_DecodeTangentZStream);
-				DecodeNormalStream(Reader, &OutMeshData.TangentsZ[VertexOffset], OutMeshData.TangentsZ.GetTypeSize(), NumChunkVertices, DecodingContext.ResidualNormalTangentZTable);
+				DecodeNormalStream(ChunkReader, &OutMeshData.TangentsZ[VertexOffset], OutMeshData.TangentsZ.GetTypeSize(), NumChunkVertices, DecodingContext.ResidualNormalTangentZTable);
 			}
 
 			if (OutMeshData.VertexInfo.bHasUV0)
 			{
 				SCOPE_CYCLE_COUNTER(STAT_DecodeUVStream);
-				DecodeUVStream(Reader, &OutMeshData.TextureCoordinates[VertexOffset], OutMeshData.TextureCoordinates.GetTypeSize(), NumChunkVertices);
+				DecodeUVStream(ChunkReader, &OutMeshData.TextureCoordinates[VertexOffset], OutMeshData.TextureCoordinates.GetTypeSize(), NumChunkVertices);
 			}
 
 			if (OutMeshData.VertexInfo.bHasMotionVectors)
 			{
 				SCOPE_CYCLE_COUNTER(STAT_DecodeMotionVectorStream);
-				DecodeMotionVectorStream(Reader, &OutMeshData.MotionVectors[VertexOffset], OutMeshData.MotionVectors.GetTypeSize(), NumChunkVertices);
+				DecodeMotionVectorStream(ChunkReader, &OutMeshData.MotionVectors[VertexOffset], OutMeshData.MotionVectors.GetTypeSize(), NumChunkVertices);
 			}
 		}
 	});
