Cherry-pick 6411590: Fix for UE-73341 geometry cache segements not initialized properly

diff --git a/Engine/Plugins/Experimental/GeometryCache/Source/GeometryCache/Private/GeometryCacheSceneProxy.cpp b/Engine/Plugins/Experimental/GeometryCache/Source/GeometryCache/Private/GeometryCacheSceneProxy.cpp
--- a/Engine/Plugins/Experimental/GeometryCache/Source/GeometryCache/Private/GeometryCacheSceneProxy.cpp
+++ b/Engine/Plugins/Experimental/GeometryCache/Source/GeometryCache/Private/GeometryCacheSceneProxy.cpp
@@ -164,6 +164,14 @@ FGeometryCacheSceneProxy::FGeometryCacheSceneProxy(UGeometryCacheComponent* Comp
 							Initializer.PrimitiveType = PT_TriangleList;
 							Initializer.bFastBuild = false;
 
+							TArray<FRayTracingGeometrySegment> Segments;
+							for (FGeometryCacheMeshBatchInfo& BatchInfo : Section->MeshData->BatchesInfo)
+							{
+								Segments.Add(FRayTracingGeometrySegment { BatchInfo.StartIndex / 3, BatchInfo.NumTriangles });
+							}
+
+							Initializer.Segments = Segments;
+
 							Section->RayTracingGeometry.SetInitializer(Initializer);
 							Section->RayTracingGeometry.InitResource();
 						}
