Custom Stencil/ Depth not working on GeometryCache actors

diff --git a/Engine/Plugins/Experimental/GeometryCache/Source/GeometryCache/Private/GeometryCacheSceneProxy.cpp b/Engine/Plugins/Experimental/GeometryCache/Source/GeometryCache/Private/GeometryCacheSceneProxy.cpp
--- a/Engine/Plugins/Experimental/GeometryCache/Source/GeometryCache/Private/GeometryCacheSceneProxy.cpp
+++ b/Engine/Plugins/Experimental/GeometryCache/Source/GeometryCache/Private/GeometryCacheSceneProxy.cpp
@@ -727,6 +727,7 @@ FPrimitiveViewRelevance FGeometryCacheSceneProxy::GetViewRelevance(const FSceneV
 	Result.bDrawRelevance = IsShown(View);
 	Result.bShadowRelevance = IsShadowCast(View);
 	Result.bDynamicRelevance = true;
+	Result.bRenderCustomDepth = ShouldRenderCustomDepth();
 	MaterialRelevance.SetPrimitiveViewRelevance(Result);
 	return Result;
 }
