Adding PostLoad routing to super class

diff --git a/Engine/Plugins/Experimental/GeometryCache/Source/GeometryCacheTracks/Private/MovieSceneGeometryCacheSection.cpp b/Engine/Plugins/Experimental/GeometryCache/Source/GeometryCacheTracks/Private/MovieSceneGeometryCacheSection.cpp
--- a/Engine/Plugins/Experimental/GeometryCache/Source/GeometryCacheTracks/Private/MovieSceneGeometryCacheSection.cpp
+++ b/Engine/Plugins/Experimental/GeometryCache/Source/GeometryCacheTracks/Private/MovieSceneGeometryCacheSection.cpp
@@ -42,6 +42,8 @@ TOptional<FFrameTime> UMovieSceneGeometryCacheSection::GetOffsetTime() const
 
 void UMovieSceneGeometryCacheSection::PostLoad()
 {
+	Super::PostLoad();
+
 	FFrameRate LegacyFrameRate = GetLegacyConversionFrameRate();
 
 	if (Params.StartOffset_DEPRECATED != GeometryCacheDeprecatedMagicNumber)
@@ -57,7 +59,6 @@ void UMovieSceneGeometryCacheSection::PostLoad()
 
 		Params.EndOffset_DEPRECATED = GeometryCacheDeprecatedMagicNumber;
 	}
-
 }
 
 void UMovieSceneGeometryCacheSection::Serialize(FArchive& Ar)
