Removing virtuals from ustruct.

These can cause a crash during serialization due to member offset differences between the base and derived types.

diff --git a/Engine/Plugins/Experimental/GeometryCache/Source/GeometryCacheTracks/Private/MovieSceneGeometryCacheTemplate.h b/Engine/Plugins/Experimental/GeometryCache/Source/GeometryCacheTracks/Private/MovieSceneGeometryCacheTemplate.h
--- a/Engine/Plugins/Experimental/GeometryCache/Source/GeometryCacheTracks/Private/MovieSceneGeometryCacheTemplate.h
+++ b/Engine/Plugins/Experimental/GeometryCache/Source/GeometryCacheTracks/Private/MovieSceneGeometryCacheTemplate.h
@@ -19,8 +19,8 @@ struct FMovieSceneGeometryCacheSectionTemplateParameters : public FMovieSceneGeo
 		, SectionStartTime(InSectionStartTime)
 		, SectionEndTime(InSectionEndTime)
 	{}
-	virtual ~FMovieSceneGeometryCacheSectionTemplateParameters() {};
-	virtual float MapTimeToAnimation(FFrameTime InPosition, FFrameRate InFrameRate) const;
+
+	float MapTimeToAnimation(FFrameTime InPosition, FFrameRate InFrameRate) const;
 
 	UPROPERTY()
 	FFrameNumber SectionStartTime;
