Fixed FGeometryCacheSceneProxy velocity view relevance.

diff --git a/Engine/Plugins/Experimental/GeometryCache/Source/GeometryCache/Private/GeometryCacheSceneProxy.cpp b/Engine/Plugins/Experimental/GeometryCache/Source/GeometryCache/Private/GeometryCacheSceneProxy.cpp
--- a/Engine/Plugins/Experimental/GeometryCache/Source/GeometryCache/Private/GeometryCacheSceneProxy.cpp
+++ b/Engine/Plugins/Experimental/GeometryCache/Source/GeometryCache/Private/GeometryCacheSceneProxy.cpp
@@ -734,6 +734,7 @@ FPrimitiveViewRelevance FGeometryCacheSceneProxy::GetViewRelevance(const FSceneV
 	Result.bShadowRelevance = IsShadowCast(View);
 	Result.bDynamicRelevance = true;
 	MaterialRelevance.SetPrimitiveViewRelevance(Result);
+	Result.bVelocityRelevance = IsMovable() && Result.bOpaqueRelevance && Result.bRenderInMainPass;
 	return Result;
 }
 
