UMovieSceneGeometryCacheSection::PostLoad() now calls Super::PostLoad(). Should fix a crash when opening the master sequence.

diff --git a/Engine/Plugins/Experimental/GeometryCache/Source/GeometryCacheTracks/Private/MovieSceneGeometryCacheSection.cpp b/Engine/Plugins/Experimental/GeometryCache/Source/GeometryCacheTracks/Private/MovieSceneGeometryCacheSection.cpp
--- a/Engine/Plugins/Experimental/GeometryCache/Source/GeometryCacheTracks/Private/MovieSceneGeometryCacheSection.cpp
+++ b/Engine/Plugins/Experimental/GeometryCache/Source/GeometryCacheTracks/Private/MovieSceneGeometryCacheSection.cpp
@@ -42,6 +42,8 @@ TOptional<FFrameTime> UMovieSceneGeometryCacheSection::GetOffsetTime() const
 
 void UMovieSceneGeometryCacheSection::PostLoad()
 {
+	Super::PostLoad();
+
 	FFrameRate LegacyFrameRate = GetLegacyConversionFrameRate();
 
 	if (Params.StartOffset_DEPRECATED != GeometryCacheDeprecatedMagicNumber)
