Add missing files for geometry cache in plugins folder

diff --git a/Engine/Plugins/Experimental/GeometryCache/Source/GeometryCache/Private/GeometryCacheSceneProxy.cpp b/Engine/Plugins/Experimental/GeometryCache/Source/GeometryCache/Private/GeometryCacheSceneProxy.cpp
--- a/Engine/Plugins/Experimental/GeometryCache/Source/GeometryCache/Private/GeometryCacheSceneProxy.cpp
+++ b/Engine/Plugins/Experimental/GeometryCache/Source/GeometryCache/Private/GeometryCacheSceneProxy.cpp
@@ -128,6 +128,49 @@ FGeometryCacheSceneProxy::FGeometryCacheSceneProxy(UGeometryCacheComponent* Comp
 			Tracks[TrackIdx] = NewSection;
 		}
 	}
+
+	if (IsRayTracingEnabled())
+	{
+		// Update at least once after the scene proxy has been constructed
+		// Otherwise it is invisible until animation starts
+		ENQUEUE_UNIQUE_RENDER_COMMAND_ONEPARAMETER(
+			FGeometryCacheUpdateAnimation,
+			FGeometryCacheSceneProxy*, SceneProxy, this,
+			{
+				SceneProxy->FrameUpdate();
+			});
+
+	#if RHI_RAYTRACING
+		if (IsRayTracingEnabled())
+		{
+			ENQUEUE_UNIQUE_RENDER_COMMAND_ONEPARAMETER(
+				FGeometryCacheInitRayTracingGeometry,
+				FGeometryCacheSceneProxy*, SceneProxy, this,
+				{
+					for (FGeomCacheTrackProxy* Section : SceneProxy->Tracks)
+					{
+						if (Section != nullptr)
+						{
+							FRayTracingGeometryInitializer Initializer;
+							const int PositionBufferIndex = Section->CurrentPositionBufferIndex != -1 ? Section->CurrentPositionBufferIndex % 2 : 0;
+							Initializer.PositionVertexBuffer = Section->PositionBuffers[PositionBufferIndex].VertexBufferRHI;
+							Initializer.IndexBuffer = Section->IndexBuffer.IndexBufferRHI;
+							Initializer.BaseVertexIndex = 0;
+							Initializer.VertexBufferStride = sizeof(FVector);
+							Initializer.VertexBufferByteOffset = 0;
+							Initializer.TotalPrimitiveCount = Section->IndexBuffer.NumIndices / 3;
+							Initializer.VertexBufferElementType = VET_Float3;
+							Initializer.PrimitiveType = PT_TriangleList;
+							Initializer.bFastBuild = false;
+
+							Section->RayTracingGeometry.SetInitializer(Initializer);
+							Section->RayTracingGeometry.InitResource();
+						}
+					}
+				});
+		}
+	#endif
+	}
 }
 
 FGeometryCacheSceneProxy::~FGeometryCacheSceneProxy()
@@ -144,6 +187,9 @@ FGeometryCacheSceneProxy::~FGeometryCacheSceneProxy()
 			Section->VertexFactory.ReleaseResource();
 			Section->PositionBuffers[0].ReleaseResource();
 			Section->PositionBuffers[1].ReleaseResource();
+		#if RHI_RAYTRACING
+			Section->RayTracingGeometry.ReleaseResource();
+		#endif
 			delete Section->MeshData;
 			if (Section->NextFrameMeshData != nullptr)
 				delete Section->NextFrameMeshData;
@@ -282,448 +328,134 @@ void FGeometryCacheSceneProxy::GetDynamicMeshElements(const TArray<const FSceneV
 
 	if (bVisible)
 	{
-		// Iterate over all batches in all tracks and add them to all the relevant views	
-		for (FGeomCacheTrackProxy* TrackProxy : Tracks )
+		if (!IsRayTracingEnabled())
 		{
-			// Render out stored TrackProxy's
-			if (TrackProxy != nullptr)
-			{	
-				const FVisibilitySample& VisibilitySample = TrackProxy->Resource->GetTrack()->GetVisibilitySample(Time, bLooping);
-				if (!VisibilitySample.bVisibilityState)
-				{
-					continue;
-				}
-
-				int32 BatchIndex = 0;
-				for (const FGeometryCacheMeshBatchInfo BatchInfo : TrackProxy->MeshData->BatchesInfo)
-				{
-					FMaterialRenderProxy* MaterialProxy = bWireframe ? WireframeMaterialInstance : TrackProxy->Materials[BatchIndex]->GetRenderProxy();
-
-					// Figure out which frame(s) we need to decode
-					int32 FrameIndex;
-					int32 NextFrameIndex;
-					float InterpolationFactor;
-					TrackProxy->Resource->GetTrack()->FindSampleIndexesFromTime(Time, bLooping, bIsPlayingBackwards, FrameIndex, NextFrameIndex, InterpolationFactor);
-					bool bDecodedAnything = false; // Did anything new get decoded this frame
-					bool bSeeked = false; // Is this frame a seek and thus the previous rendered frame's data invalid
-					bool bDecoderError = false; // If we have a decoder error we don't interpolate and we don't update the vertex buffers
-												// so essentially we just keep the last valid frame...
-
-					// Compare this against the frames we got and keep some/all/none of them
-					// This will work across frames but also within a frame if the mesh is in several views
-					if (TrackProxy->FrameIndex != FrameIndex || TrackProxy->NextFrameIndex != NextFrameIndex)
-					{
-						// Normal case the next frame is the new current frame
-						if (TrackProxy->NextFrameIndex == FrameIndex)
-						{
-							// Cycle the current and next frame double buffer
-							FGeometryCacheMeshData *OldFrameMesh = TrackProxy->MeshData;
-							TrackProxy->MeshData = TrackProxy->NextFrameMeshData;
-							TrackProxy->NextFrameMeshData = OldFrameMesh;
-
-							int32 OldFrameIndex = TrackProxy->FrameIndex;
-							TrackProxy->FrameIndex = TrackProxy->NextFrameIndex;
-							TrackProxy->NextFrameIndex = OldFrameIndex;
-
-							// Decode the new next frame
-							if (TrackProxy->Resource->DecodeMeshData(NextFrameIndex, *TrackProxy->NextFrameMeshData))
-							{		
-								bDecodedAnything = true;
-								// Only register this if we actually successfully decoded
-								TrackProxy->NextFrameIndex = NextFrameIndex;
-							}
-							else
-							{
-								// Mark the frame as corrupted
-								TrackProxy->NextFrameIndex = -1;
-								bDecoderError = true;
-							}
-						}
-						// Probably a seek or the mesh hasn't been visible in a while decode two frames
-						else
-						{
-							if (TrackProxy->Resource->DecodeMeshData(FrameIndex, *TrackProxy->MeshData))
-							{
-								TrackProxy->NextFrameMeshData->Indices = TrackProxy->MeshData->Indices;
-								if (TrackProxy->Resource->DecodeMeshData(NextFrameIndex, *TrackProxy->NextFrameMeshData))
-								{
-									TrackProxy->FrameIndex = FrameIndex;
-									TrackProxy->NextFrameIndex = NextFrameIndex;
-									bSeeked = true;
-									bDecodedAnything = true;
-								}
-								else
-								{
-									// The first frame decoded fine but the second didn't 
-									// we need to specially handle this
-									TrackProxy->NextFrameIndex = -1;
-									bDecoderError = true;
-								}
-							}
-							else
-							{
-								TrackProxy->FrameIndex = -1;
-								bDecoderError = true;
-							}
-						}
-					}
+			// When ray tracing is disabled, update only when visible
+			// This is the old behavior
+			FrameUpdate();
+		}
 
-					// Check if we can interpolate between the two frames we have available
-					const bool bCanInterpolate = TrackProxy->Resource->IsTopologyCompatible(TrackProxy->FrameIndex, TrackProxy->NextFrameIndex);
+		// Iterate over all batches in all tracks and add them to all the relevant views	
+		for (const FGeomCacheTrackProxy* TrackProxy : Tracks)
+		{
+			const FVisibilitySample& VisibilitySample = TrackProxy->Resource->GetTrack()->GetVisibilitySample(Time, bLooping);
+			if (!VisibilitySample.bVisibilityState)
+			{
+				continue;
+			}
 
-					// Check if we have explicit motion vectors
-					const bool bHasMotionVectors = (
-						TrackProxy->MeshData->VertexInfo.bHasMotionVectors &&
-						TrackProxy->NextFrameMeshData->VertexInfo.bHasMotionVectors &&
-						TrackProxy->MeshData->Positions.Num() == TrackProxy->MeshData->MotionVectors.Num())
-						&& (TrackProxy->NextFrameMeshData->Positions.Num() == TrackProxy->NextFrameMeshData->MotionVectors.Num());
+			const int32 NumBatches = TrackProxy->MeshData->BatchesInfo.Num();
+			const bool bHasMotionVectors = (
+				TrackProxy->MeshData->VertexInfo.bHasMotionVectors &&
+				TrackProxy->NextFrameMeshData->VertexInfo.bHasMotionVectors &&
+				TrackProxy->MeshData->Positions.Num() == TrackProxy->MeshData->MotionVectors.Num())
+				&& (TrackProxy->NextFrameMeshData->Positions.Num() == TrackProxy->NextFrameMeshData->MotionVectors.Num());
 
-					// Can we interpolate the vertex data?
-					if (bCanInterpolate && !bDecoderError && CVarInterpolateFrames.GetValueOnRenderThread() != 0)
-					{
-						// Interpolate if the time has changed.
-						// note: This is a bit precarious as this code is called multiple times per frame. This ensures
-						// we only interpolate once (which is a nice optimization) but more importantly that we only
-						// bump the CurrentPositionBufferIndex once per frame. This ensures that last frame's position
-						// buffer is not overwritten.
-						// If motion blur suddenly seems to stop working while it should be working it may be that the
-						// CurrentPositionBufferIndex gets inadvertently bumped twice per frame essentially using the same
-						// data for current and previous during rendering.
-						if (TrackProxy->PositionBufferFrameTimes[TrackProxy->CurrentPositionBufferIndex % 2] != Time)
-						{
-							TArray<FVector> InterpolatedPositions;
-							TArray<FPackedNormal> InterpolatedTangentX;
-							TArray<FPackedNormal> InterpolatedTangentZ;
-							TArray<FVector2D> InterpolatedUVs;
-							TArray<FColor> InterpolatedColors;
-
-							const int32 NumVerts = TrackProxy->MeshData->Positions.Num();
-
-							InterpolatedPositions.AddUninitialized(NumVerts);
-							InterpolatedTangentX.AddUninitialized(NumVerts);
-							InterpolatedTangentZ.AddUninitialized(NumVerts);
-							InterpolatedUVs.AddUninitialized(NumVerts);
-							InterpolatedColors.AddUninitialized(NumVerts);
-
-							TArray<FVector> InterpolatedMotionVectors;
-							if (bHasMotionVectors)
-							{
-								InterpolatedMotionVectors.AddUninitialized(NumVerts);
-							}
-
-
-							const float OneMinusInterp = 1.0 - InterpolationFactor;
-							const int32 InterpFixed = (int32)(InterpolationFactor * 255.0f);
-							const int32 OneMinusInterpFixed = 255 - InterpFixed;
-
-							for (int32 Index = 0; Index < NumVerts; ++Index)
-								{
-								const FVector& PositionA = TrackProxy->MeshData->Positions[Index];
-								const FVector& PositionB = TrackProxy->NextFrameMeshData->Positions[Index];
-								InterpolatedPositions[Index] = PositionA * OneMinusInterp + PositionB* InterpolationFactor;
-							}
-
-							for (int32 Index = 0; Index < NumVerts; ++Index)
-							{
-								// The following are already 8 bit so quantized enough we can do exact equal comparisons
-								const FPackedNormal& TangentXA = TrackProxy->MeshData->TangentsX[Index];
-								const FPackedNormal& TangentXB = TrackProxy->NextFrameMeshData->TangentsX[Index];
-								const FPackedNormal& TangentZA = TrackProxy->MeshData->TangentsZ[Index];
-								const FPackedNormal& TangentZB = TrackProxy->NextFrameMeshData->TangentsZ[Index];
-
-								InterpolatedTangentX[Index] = InterpolatePackedNormal(TangentXA, TangentXB, InterpFixed, OneMinusInterpFixed);
-								InterpolatedTangentZ[Index] = InterpolatePackedNormal(TangentZA, TangentZB, InterpFixed, OneMinusInterpFixed);
-							}
-
-							if (TrackProxy->MeshData->VertexInfo.bHasColor0) for (int32 Index = 0; Index < NumVerts; ++Index)
-							{
-								const FColor& ColorA = TrackProxy->MeshData->Colors[Index];
-								const FColor& ColorB = TrackProxy->NextFrameMeshData->Colors[Index];
-								InterpolatedColors[Index] = InterpolatePackedColor(ColorA, ColorB, InterpFixed, OneMinusInterpFixed);
-							}
-
-							if (TrackProxy->MeshData->VertexInfo.bHasUV0) for (int32 Index = 0; Index < NumVerts; ++Index)
-							{
-								const FVector2D& UVA = TrackProxy->MeshData->TextureCoordinates[Index];
-								const FVector2D& UVB = TrackProxy->NextFrameMeshData->TextureCoordinates[Index];
-
-								InterpolatedUVs[Index] = UVA * OneMinusInterp + UVB * InterpolationFactor;
-							}
-
-							if (bHasMotionVectors) for (int32 Index = 0; Index < NumVerts; ++Index)
-							{
-								InterpolatedMotionVectors[Index] = TrackProxy->MeshData->MotionVectors[Index] * OneMinusInterp +
-									TrackProxy->NextFrameMeshData->MotionVectors[Index] * InterpolationFactor;
-								}
-
-								// Upload other non-motionblurred data
-							if (!TrackProxy->MeshData->VertexInfo.bConstantIndices)
-								TrackProxy->IndexBuffer.Update(TrackProxy->MeshData->Indices);
-
-							if (TrackProxy->MeshData->VertexInfo.bHasTangentX)
-								TrackProxy->TangentXBuffer.Update(InterpolatedTangentX);
-							if (TrackProxy->MeshData->VertexInfo.bHasTangentZ)
-								TrackProxy->TangentZBuffer.Update(InterpolatedTangentZ);
-
-							if (TrackProxy->MeshData->VertexInfo.bHasUV0)
-								TrackProxy->TextureCoordinatesBuffer.Update(InterpolatedUVs);
-
-							if (TrackProxy->MeshData->VertexInfo.bHasColor0)
-								TrackProxy->ColorBuffer.Update(InterpolatedColors);
-
-							bool bIsCompatibleWithCachedFrame = TrackProxy->Resource->IsTopologyCompatible(
-								TrackProxy->PositionBufferFrameIndices[TrackProxy->CurrentPositionBufferIndex % 2],
-								TrackProxy->FrameIndex);
-
-							if (!bHasMotionVectors)
-							{
-								// Initialize both buffers the first frame
-								if (TrackProxy->CurrentPositionBufferIndex == -1 || !bIsCompatibleWithCachedFrame)
-								{
-									TrackProxy->PositionBuffers[0].Update(InterpolatedPositions);
-									TrackProxy->PositionBuffers[1].Update(InterpolatedPositions);
-									TrackProxy->CurrentPositionBufferIndex = 0;
-									TrackProxy->PositionBufferFrameTimes[0] = Time;
-									TrackProxy->PositionBufferFrameTimes[1] = Time;
-									// We need to keep a frame index in order to ensure topology consistency. As we can interpolate 
-									// FrameIndex and NextFrameIndex are certainly topo-compatible so it doesn't really matter which 
-									// one we keep here. But wee keep NextFrameIndex as that is most useful to validate against
-									// the frame coming up
-									TrackProxy->PositionBufferFrameIndices[0] = TrackProxy->NextFrameIndex;
-									TrackProxy->PositionBufferFrameIndices[1] = TrackProxy->NextFrameIndex;
-								}
-								else
-								{
-									TrackProxy->CurrentPositionBufferIndex++;
-									TrackProxy->PositionBuffers[TrackProxy->CurrentPositionBufferIndex % 2].Update(InterpolatedPositions);
-									TrackProxy->PositionBufferFrameTimes[TrackProxy->CurrentPositionBufferIndex % 2] = Time;
-									TrackProxy->PositionBufferFrameIndices[TrackProxy->CurrentPositionBufferIndex % 2] = TrackProxy->NextFrameIndex;
-								}
-							}
-							else
-							{
-								TrackProxy->CurrentPositionBufferIndex = 0;
-								TrackProxy->PositionBuffers[0].Update(InterpolatedPositions);
-								TrackProxy->PositionBuffers[1].Update(InterpolatedMotionVectors);
-								TrackProxy->PositionBufferFrameIndices[0] = TrackProxy->FrameIndex;
-								TrackProxy->PositionBufferFrameIndices[1] = -1;
-								TrackProxy->PositionBufferFrameTimes[0] = Time;
-								TrackProxy->PositionBufferFrameTimes[1] = Time;
-							}
-						}
+			for (int32 BatchIndex = 0; BatchIndex < NumBatches; ++BatchIndex)
+			{
+				FMaterialRenderProxy* MaterialProxy = bWireframe ? WireframeMaterialInstance : TrackProxy->Materials[BatchIndex]->GetRenderProxy();
+				const FGeometryCacheMeshBatchInfo BatchInfo = TrackProxy->MeshData->BatchesInfo[BatchIndex];
 
+				FGeometryCacheVertexFactoryUserDataWrapper &UserDataWrapper = Collector.AllocateOneFrameResource<FGeometryCacheVertexFactoryUserDataWrapper>();
+				FGeometryCacheVertexFactoryUserData &UserData = UserDataWrapper.Data;
 
-						}
-						else
-						{
-							// We just don't interpolate between frames if we got GPU to burn we could someday render twice and stipple fade between it :-D like with lods
-
-							// Only bother uploading if anything changed or when the we failed to decode anything make sure update the gpu buffers regardless
-							if (bDecodedAnything || bDecoderError)
-							{
-								//TrackProxy->IndexBuffer.Update(TrackProxy->MeshData->Indices);
-
-								const int32 NumVertices = TrackProxy->MeshData->Positions.Num();
-
-								if (TrackProxy->MeshData->VertexInfo.bHasTangentX)
-									TrackProxy->TangentXBuffer.Update(TrackProxy->MeshData->TangentsX);
-								if (TrackProxy->MeshData->VertexInfo.bHasTangentZ)
-									TrackProxy->TangentZBuffer.Update(TrackProxy->MeshData->TangentsZ);
-
-								if (!TrackProxy->MeshData->VertexInfo.bConstantIndices)
-									TrackProxy->IndexBuffer.Update(TrackProxy->MeshData->Indices);
-
-								if (TrackProxy->MeshData->VertexInfo.bHasUV0)
-									TrackProxy->TextureCoordinatesBuffer.Update(TrackProxy->MeshData->TextureCoordinates);
-
-								if (TrackProxy->MeshData->VertexInfo.bHasColor0)
-									TrackProxy->ColorBuffer.Update(TrackProxy->MeshData->Colors);
-
-									bool bIsCompatibleWithCachedFrame = TrackProxy->Resource->IsTopologyCompatible(
-										TrackProxy->PositionBufferFrameIndices[TrackProxy->CurrentPositionBufferIndex % 2],
-										TrackProxy->FrameIndex);
-
-								if (!bHasMotionVectors)
-								{
-									// Initialize both buffers the first frame or when topology changed as we can't render
-									// with a previous buffer referencing a buffer from another topology
-									if (TrackProxy->CurrentPositionBufferIndex == -1 || !bIsCompatibleWithCachedFrame || bSeeked)
-									{
-										TrackProxy->PositionBuffers[0].Update(TrackProxy->MeshData->Positions);
-										TrackProxy->PositionBuffers[1].Update(TrackProxy->MeshData->Positions);
-										TrackProxy->CurrentPositionBufferIndex = 0;
-										TrackProxy->PositionBufferFrameIndices[0] = TrackProxy->FrameIndex;
-										TrackProxy->PositionBufferFrameIndices[1] = TrackProxy->FrameIndex;
-									}
-									// We still use the previous frame's buffer as a motion blur previous position. As interpolation is switched
-									// off the actual time of this previous frame depends on the geometry cache framerate and playback speed
-									// so the motion blur vectors may not really be anything relevant. Do we want to just disable motion blur? 
-									// But as an optimization skipping interpolation when the cache fps is near to the actual game fps this is obviously nice...
-									else
-									{
-										TrackProxy->CurrentPositionBufferIndex++;
-										TrackProxy->PositionBuffers[TrackProxy->CurrentPositionBufferIndex % 2].Update(TrackProxy->MeshData->Positions);
-										TrackProxy->PositionBufferFrameIndices[TrackProxy->CurrentPositionBufferIndex % 2] = TrackProxy->FrameIndex;
-									}
-								}
-								else
-								{
-									TrackProxy->CurrentPositionBufferIndex = 0;
-									TrackProxy->PositionBuffers[0].Update(TrackProxy->MeshData->Positions);
-									TrackProxy->PositionBuffers[1].Update(TrackProxy->MeshData->MotionVectors);
-									TrackProxy->PositionBufferFrameIndices[0] = TrackProxy->FrameIndex;
-									TrackProxy->PositionBufferFrameIndices[1] = -1;
-									TrackProxy->PositionBufferFrameTimes[0] = Time;
-									TrackProxy->PositionBufferFrameTimes[1] = Time;
-								}
-							}
-						}
-
-						#if 0
-						bool bOffloadUpdate = CVarOffloadUpdate.GetValueOnRenderThread() != 0;
-						if (TrackProxy->SampleIndex != TrackProxy->UploadedSampleIndex)
-						{
-							TrackProxy->UploadedSampleIndex = TrackProxy->SampleIndex;
-
-							if (bOffloadUpdate)
-							{
-								check(false);
-								// Only update the size on this thread
-								TrackProxy->IndexBuffer.UpdateSizeOnly(TrackProxy->MeshData->Indices.Num());
-								TrackProxy->VertexBuffer.UpdateSizeTyped<FNoPositionVertex>(TrackProxy->MeshData->Vertices.Num());
-								
-								// Do the interpolation on a worker thread
-								FGraphEventRef CompletionFence = FFunctionGraphTask::CreateAndDispatchWhenReady([]()
-								{
-
-								}, GET_STATID(STAT_BufferUpdateTask), NULL, ENamedThreads::AnyThread);
-
-								// Queue a command on the RHI thread that waits for the interpolation job and then uploads them to the GPU
-								FRHICommandListImmediate &RHICommandList = GetImmediateCommandList_ForRenderCommand();
-								new (RHICommandList.AllocCommand<FRHICommandUpdateGeometryCacheBuffer>())FRHICommandUpdateGeometryCacheBuffer(
-									CompletionFence,
-									TrackProxy->VertexBuffer.VertexBufferRHI,
-									TrackProxy->MeshData->Vertices.GetData(),
-									TrackProxy->VertexBuffer.GetSizeInBytes(),
-									TrackProxy->IndexBuffer.IndexBufferRHI,
-									TrackProxy->MeshData->Indices.GetData(),
-									TrackProxy->IndexBuffer.SizeInBytes());
-
-								// Upload vertex data
-								/*void* RESTRICT Data = (void* RESTRICT)GDynamicRHI->RHILockVertexBuffer(TrackProxy->VertexBuffer.VertexBufferRHI, 0, TrackProxy->VertexBuffer.SizeInBytes(), RLM_WriteOnly);
-								FMemory::BigBlockMemcpy(Data, TrackProxy->MeshData->Vertices.GetData(), TrackProxy->VertexBuffer.SizeInBytes());
-								GDynamicRHI->RHIUnlockVertexBuffer(TrackProxy->VertexBuffer.VertexBufferRHI);
-
-								// Upload index data
-								Data = (void* RESTRICT)GDynamicRHI->RHILockIndexBuffer(TrackProxy->IndexBuffer.IndexBufferRHI, 0, TrackProxy->IndexBuffer.SizeInBytes(), RLM_WriteOnly);
-								FMemory::BigBlockMemcpy(Data, TrackProxy->MeshData->Indices.GetData(), TrackProxy->IndexBuffer.SizeInBytes());
-								GDynamicRHI->RHIUnlockIndexBuffer(TrackProxy->IndexBuffer.IndexBufferRHI);*/
-							}
-							else
-							{
-								TrackProxy->IndexBuffer.Update(TrackProxy->MeshData->Indices);
-								TrackProxy->VertexBuffer.Update(TrackProxy->MeshData->Vertices);
-
-								// Initialize both buffers the first frame
-								if (TrackProxy->CurrentPositionBufferIndex == -1)
-								{
-									TrackProxy->PositonBuffers[0].Update(TrackProxy->MeshData->Vertices);
-									TrackProxy->PositonBuffers[1].Update(TrackProxy->MeshData->Vertices);
-									TrackProxy->CurrentPositionBufferIndex = 0;
-								}
-								else
-								{
-									TrackProxy->CurrentPositionBufferIndex++;
-									TrackProxy->PositonBuffers[TrackProxy->CurrentPositionBufferIndex%2].Update(TrackProxy->MeshData->Vertices);
-								}
-							}
-						}
-	#endif
-					}
+				UserData.MeshExtension = FVector::OneVector;
+				UserData.MeshOrigin = FVector::ZeroVector;
+				if (!bHasMotionVectors)
+				{
+					UserData.MotionBlurDataExtension = FVector::OneVector;
+					UserData.MotionBlurDataOrigin = FVector::ZeroVector;
+					UserData.MotionBlurPositionScale = 0.0f;
 				}
-			}
-		
-			for (const FGeomCacheTrackProxy* TrackProxy : Tracks)
-			{	
-				const FVisibilitySample& VisibilitySample = TrackProxy->Resource->GetTrack()->GetVisibilitySample(Time, bLooping);
-				if (!VisibilitySample.bVisibilityState)
+				else
 				{
-					continue;
+					UserData.MotionBlurDataExtension = FVector::OneVector * PlaybackSpeed;
+					UserData.MotionBlurDataOrigin = FVector::ZeroVector;
+					UserData.MotionBlurPositionScale = 1.0f;
 				}
 
-				const int32 NumBatches = TrackProxy->MeshData->BatchesInfo.Num();
-				const bool bHasMotionVectors = (
-					TrackProxy->MeshData->VertexInfo.bHasMotionVectors &&
-					TrackProxy->NextFrameMeshData->VertexInfo.bHasMotionVectors &&
-					TrackProxy->MeshData->Positions.Num() == TrackProxy->MeshData->MotionVectors.Num())
-					&& (TrackProxy->NextFrameMeshData->Positions.Num() == TrackProxy->NextFrameMeshData->MotionVectors.Num());
-	
-				for (int32 BatchIndex = 0; BatchIndex < NumBatches; ++BatchIndex)
-				{	
-					FMaterialRenderProxy* MaterialProxy = bWireframe ? WireframeMaterialInstance : TrackProxy->Materials[BatchIndex]->GetRenderProxy();
-					const FGeometryCacheMeshBatchInfo BatchInfo = TrackProxy->MeshData->BatchesInfo[BatchIndex];
-
-					FGeometryCacheVertexFactoryUserDataWrapper &UserDataWrapper = Collector.AllocateOneFrameResource<FGeometryCacheVertexFactoryUserDataWrapper>();
-					FGeometryCacheVertexFactoryUserData &UserData = UserDataWrapper.Data;
+				if (IsRayTracingEnabled())
+				{
+					// No vertex manipulation is allowed in the vertex shader
+					// Otherwise we need an additional compute shader pass to execute the vertex shader and dump to a staging buffer
+					check(UserData.MeshExtension == FVector::OneVector);
+					check(UserData.MeshOrigin == FVector::ZeroVector);
+				}
 
-					UserData.MeshExtension = FVector::OneVector;
-					UserData.MeshOrigin = FVector::ZeroVector;
-					if (!bHasMotionVectors)
-					{
-						UserData.MotionBlurDataExtension = FVector::OneVector;
-						UserData.MotionBlurDataOrigin = FVector::ZeroVector;
-						UserData.MotionBlurPositionScale = 0.0f;
-					}
-					else
+				UserData.PositionBuffer = &TrackProxy->PositionBuffers[TrackProxy->CurrentPositionBufferIndex % 2];
+				UserData.MotionBlurDataBuffer = &TrackProxy->PositionBuffers[(TrackProxy->CurrentPositionBufferIndex+1) % 2];
+
+				FGeometryCacheVertexFactoryUniformBufferParameters UniformBufferParameters;
+
+				UniformBufferParameters.MeshOrigin = UserData.MeshOrigin;
+				UniformBufferParameters.MeshExtension = UserData.MeshExtension;
+				UniformBufferParameters.MotionBlurDataOrigin = UserData.MotionBlurDataOrigin;
+				UniformBufferParameters.MotionBlurDataExtension = UserData.MotionBlurDataExtension;
+				UniformBufferParameters.MotionBlurPositionScale = UserData.MotionBlurPositionScale;
+
+				UserData.UniformBuffer = FGeometryCacheVertexFactoryUniformBufferParametersRef::CreateUniformBufferImmediate(UniformBufferParameters, UniformBuffer_SingleFrame);
+				TrackProxy->VertexFactory.CreateManualVertexFetchUniformBuffer(UserData.PositionBuffer, UserData.MotionBlurDataBuffer, UserData);
+
+				// Draw the mesh.
+				FMeshBatch& Mesh = Collector.AllocateMesh();
+				FMeshBatchElement& BatchElement = Mesh.Elements[0];
+				BatchElement.IndexBuffer = &TrackProxy->IndexBuffer;
+				Mesh.bWireframe = bWireframe;
+				Mesh.VertexFactory = &TrackProxy->VertexFactory;
+				Mesh.MaterialRenderProxy = MaterialProxy;
+				Mesh.SegmentIndex = 0;
+
+				const FMatrix& LocalToWorldTransform = TrackProxy->WorldMatrix * GetLocalToWorld();
+
+				FDynamicPrimitiveUniformBuffer& DynamicPrimitiveUniformBuffer = Collector.AllocateOneFrameResource<FDynamicPrimitiveUniformBuffer>();
+				DynamicPrimitiveUniformBuffer.Set(LocalToWorldTransform, LocalToWorldTransform, GetBounds(), GetLocalBounds(), true, false, UseEditorDepthTest());
+				BatchElement.PrimitiveUniformBuffer = DynamicPrimitiveUniformBuffer.UniformBuffer.GetUniformBufferRHI();
+
+				BatchElement.FirstIndex = BatchInfo.StartIndex;
+				BatchElement.NumPrimitives = BatchInfo.NumTriangles;
+				BatchElement.MinVertexIndex = 0;
+				BatchElement.MaxVertexIndex = TrackProxy->MeshData->Positions.Num() - 1;
+				BatchElement.VertexFactoryUserData = &UserDataWrapper.Data;
+				Mesh.ReverseCulling = IsLocalToWorldDeterminantNegative();
+				Mesh.Type = PT_TriangleList;
+				Mesh.DepthPriorityGroup = SDPG_World;
+				Mesh.bCanApplyViewModeOverrides = false;
+
+				for (int32 ViewIndex = 0; ViewIndex < Views.Num(); ViewIndex++)
+				{
+					if (VisibilityMap & (1 << ViewIndex))
 					{
-						UserData.MotionBlurDataExtension = FVector::OneVector * PlaybackSpeed;
-						UserData.MotionBlurDataOrigin = FVector::ZeroVector;
-						UserData.MotionBlurPositionScale = 1.0f;
-					}
-						
-					UserData.PositionBuffer = &TrackProxy->PositionBuffers[TrackProxy->CurrentPositionBufferIndex % 2];
-					UserData.MotionBlurDataBuffer = &TrackProxy->PositionBuffers[(TrackProxy->CurrentPositionBufferIndex+1) % 2];
-
-					// Draw the mesh.
-					FMeshBatch& Mesh = Collector.AllocateMesh();
-					FMeshBatchElement& BatchElement = Mesh.Elements[0];
-					BatchElement.IndexBuffer = &TrackProxy->IndexBuffer;
-					Mesh.bWireframe = bWireframe;
-					Mesh.VertexFactory = &TrackProxy->VertexFactory;
-					Mesh.MaterialRenderProxy = MaterialProxy;
-                    
-                    const FMatrix& LocalToWorldTransform = TrackProxy->WorldMatrix * GetLocalToWorld();
-                    
-                    FDynamicPrimitiveUniformBuffer& DynamicPrimitiveUniformBuffer = Collector.AllocateOneFrameResource<FDynamicPrimitiveUniformBuffer>();
-                    DynamicPrimitiveUniformBuffer.Set(LocalToWorldTransform, LocalToWorldTransform, GetBounds(), GetLocalBounds(), true, false, UseEditorDepthTest());
-                    BatchElement.PrimitiveUniformBuffer = DynamicPrimitiveUniformBuffer.UniformBuffer.GetUniformBufferRHI();
-                    
-					BatchElement.FirstIndex = BatchInfo.StartIndex;
-					BatchElement.NumPrimitives = BatchInfo.NumTriangles;
-					BatchElement.MinVertexIndex = 0;
-					BatchElement.MaxVertexIndex = TrackProxy->MeshData->Positions.Num() - 1;
-					BatchElement.VertexFactoryUserData = &UserDataWrapper.Data;
-					Mesh.ReverseCulling = IsLocalToWorldDeterminantNegative();
-					Mesh.Type = PT_TriangleList;
-					Mesh.DepthPriorityGroup = SDPG_World;
-					Mesh.bCanApplyViewModeOverrides = false;
-
-					for (int32 ViewIndex = 0; ViewIndex < Views.Num(); ViewIndex++)
-					{
-						if (VisibilityMap & (1 << ViewIndex))
-						{
-							Collector.AddMesh(ViewIndex, Mesh);
-							INC_DWORD_STAT_BY(STAT_GeometryCacheSceneProxy_TriangleCount, BatchElement.NumPrimitives);
-							INC_DWORD_STAT_BY(STAT_GeometryCacheSceneProxy_MeshBatchCount, 1);
-
-#if !(UE_BUILD_SHIPPING || UE_BUILD_TEST)
-							// Render bounds
-							RenderBounds(Collector.GetPDI(ViewIndex), ViewFamily.EngineShowFlags, GetBounds(), IsSelected());
-#endif
-						}			
+						Collector.AddMesh(ViewIndex, Mesh);
+						INC_DWORD_STAT_BY(STAT_GeometryCacheSceneProxy_TriangleCount, BatchElement.NumPrimitives);
+						INC_DWORD_STAT_BY(STAT_GeometryCacheSceneProxy_MeshBatchCount, 1);
+
+					#if !(UE_BUILD_SHIPPING || UE_BUILD_TEST)
+						// Render bounds
+						RenderBounds(Collector.GetPDI(ViewIndex), ViewFamily.EngineShowFlags, GetBounds(), IsSelected());
+					#endif
 					}
 				}
-			}			
+			}
+		}
+	}
+}
+
+void FGeometryCacheSceneProxy::GetRayTracingGeometryInstances(TArray<FRayTracingGeometryInstanceCollection>& OutInstanceCollections)
+{
+	for (FGeomCacheTrackProxy* TrackProxy : Tracks)
+	{
+		const FVisibilitySample& VisibilitySample = TrackProxy->Resource->GetTrack()->GetVisibilitySample(Time, bLooping);
+		if (!VisibilitySample.bVisibilityState)
+		{
+			continue;
+		}
+
+		FRayTracingGeometryInstanceCollection Collection;
+		Collection.Geometry = &TrackProxy->RayTracingGeometry;
+		Collection.DynamicVertexPositionBuffer = nullptr;
+		Collection.InstanceTransformMode = FRayTracingGeometryInstanceCollection::TransformMode::InheritFromSceneProxy;
+
+		OutInstanceCollections.Add(Collection);
 	}
 }
 
@@ -764,6 +496,384 @@ void FGeometryCacheSceneProxy::UpdateAnimation(float NewTime, bool bNewLooping,
 	bLooping = bNewLooping;
 	bIsPlayingBackwards = bNewIsPlayingBackwards;
 	PlaybackSpeed = NewPlaybackSpeed;
+
+	if (IsRayTracingEnabled())
+	{
+		// When ray tracing is enabled, update regardless of visibility
+		FrameUpdate();
+
+	#if RHI_RAYTRACING
+		for (FGeomCacheTrackProxy* Section : Tracks)
+		{
+			if (Section != nullptr)
+			{
+				const int PositionBufferIndex = Section->CurrentPositionBufferIndex != -1 ? Section->CurrentPositionBufferIndex % 2 : 0;
+
+				Section->RayTracingGeometry.Initializer.PositionVertexBuffer = Section->PositionBuffers[PositionBufferIndex].VertexBufferRHI;
+				Section->RayTracingGeometry.Initializer.TotalPrimitiveCount = Section->IndexBuffer.NumIndices / 3;
+				Section->RayTracingGeometry.UpdateRHI();
+			}
+		}
+	#endif
+	}
+}
+
+void FGeometryCacheSceneProxy::FrameUpdate() const
+{
+	for (FGeomCacheTrackProxy* TrackProxy : Tracks)
+	{
+		// Render out stored TrackProxy's
+		if (TrackProxy != nullptr)
+		{
+			const FVisibilitySample& VisibilitySample = TrackProxy->Resource->GetTrack()->GetVisibilitySample(Time, bLooping);
+			if (!VisibilitySample.bVisibilityState)
+			{
+				continue;
+			}
+
+			// Figure out which frame(s) we need to decode
+			int32 FrameIndex;
+			int32 NextFrameIndex;
+			float InterpolationFactor;
+			TrackProxy->Resource->GetTrack()->FindSampleIndexesFromTime(Time, bLooping, bIsPlayingBackwards, FrameIndex, NextFrameIndex, InterpolationFactor);
+			bool bDecodedAnything = false; // Did anything new get decoded this frame
+			bool bSeeked = false; // Is this frame a seek and thus the previous rendered frame's data invalid
+			bool bDecoderError = false; // If we have a decoder error we don't interpolate and we don't update the vertex buffers
+										// so essentially we just keep the last valid frame...
+
+			// Compare this against the frames we got and keep some/all/none of them
+			// This will work across frames but also within a frame if the mesh is in several views
+			if (TrackProxy->FrameIndex != FrameIndex || TrackProxy->NextFrameIndex != NextFrameIndex)
+			{
+				// Normal case the next frame is the new current frame
+				if (TrackProxy->NextFrameIndex == FrameIndex)
+				{
+					// Cycle the current and next frame double buffer
+					FGeometryCacheMeshData *OldFrameMesh = TrackProxy->MeshData;
+					TrackProxy->MeshData = TrackProxy->NextFrameMeshData;
+					TrackProxy->NextFrameMeshData = OldFrameMesh;
+
+					int32 OldFrameIndex = TrackProxy->FrameIndex;
+					TrackProxy->FrameIndex = TrackProxy->NextFrameIndex;
+					TrackProxy->NextFrameIndex = OldFrameIndex;
+
+					// Decode the new next frame
+					if (TrackProxy->Resource->DecodeMeshData(NextFrameIndex, *TrackProxy->NextFrameMeshData))
+					{
+						bDecodedAnything = true;
+						// Only register this if we actually successfully decoded
+						TrackProxy->NextFrameIndex = NextFrameIndex;
+					}
+					else
+					{
+						// Mark the frame as corrupted
+						TrackProxy->NextFrameIndex = -1;
+						bDecoderError = true;
+					}
+				}
+				// Probably a seek or the mesh hasn't been visible in a while decode two frames
+				else
+				{
+					if (TrackProxy->Resource->DecodeMeshData(FrameIndex, *TrackProxy->MeshData))
+					{
+						TrackProxy->NextFrameMeshData->Indices = TrackProxy->MeshData->Indices;
+						if (TrackProxy->Resource->DecodeMeshData(NextFrameIndex, *TrackProxy->NextFrameMeshData))
+						{
+							TrackProxy->FrameIndex = FrameIndex;
+							TrackProxy->NextFrameIndex = NextFrameIndex;
+							bSeeked = true;
+							bDecodedAnything = true;
+						}
+						else
+						{
+							// The first frame decoded fine but the second didn't 
+							// we need to specially handle this
+							TrackProxy->NextFrameIndex = -1;
+							bDecoderError = true;
+						}
+					}
+					else
+					{
+						TrackProxy->FrameIndex = -1;
+						bDecoderError = true;
+					}
+				}
+			}
+
+			// Check if we can interpolate between the two frames we have available
+			const bool bCanInterpolate = TrackProxy->Resource->IsTopologyCompatible(TrackProxy->FrameIndex, TrackProxy->NextFrameIndex);
+
+			// Check if we have explicit motion vectors
+			const bool bHasMotionVectors = (
+				TrackProxy->MeshData->VertexInfo.bHasMotionVectors &&
+				TrackProxy->NextFrameMeshData->VertexInfo.bHasMotionVectors &&
+				TrackProxy->MeshData->Positions.Num() == TrackProxy->MeshData->MotionVectors.Num())
+				&& (TrackProxy->NextFrameMeshData->Positions.Num() == TrackProxy->NextFrameMeshData->MotionVectors.Num());
+
+			// Can we interpolate the vertex data?
+			if (bCanInterpolate && !bDecoderError && CVarInterpolateFrames.GetValueOnRenderThread() != 0)
+			{
+				// Interpolate if the time has changed.
+				// note: This is a bit precarious as this code is called multiple times per frame. This ensures
+				// we only interpolate once (which is a nice optimization) but more importantly that we only
+				// bump the CurrentPositionBufferIndex once per frame. This ensures that last frame's position
+				// buffer is not overwritten.
+				// If motion blur suddenly seems to stop working while it should be working it may be that the
+				// CurrentPositionBufferIndex gets inadvertently bumped twice per frame essentially using the same
+				// data for current and previous during rendering.
+				if (TrackProxy->PositionBufferFrameTimes[TrackProxy->CurrentPositionBufferIndex % 2] != Time)
+				{
+					TArray<FVector> InterpolatedPositions;
+					TArray<FPackedNormal> InterpolatedTangentX;
+					TArray<FPackedNormal> InterpolatedTangentZ;
+					TArray<FVector2D> InterpolatedUVs;
+					TArray<FColor> InterpolatedColors;
+
+					const int32 NumVerts = TrackProxy->MeshData->Positions.Num();
+
+					InterpolatedPositions.AddUninitialized(NumVerts);
+					InterpolatedTangentX.AddUninitialized(NumVerts);
+					InterpolatedTangentZ.AddUninitialized(NumVerts);
+					InterpolatedUVs.AddUninitialized(NumVerts);
+					InterpolatedColors.AddUninitialized(NumVerts);
+
+					TArray<FVector> InterpolatedMotionVectors;
+					if (bHasMotionVectors)
+					{
+						InterpolatedMotionVectors.AddUninitialized(NumVerts);
+					}
+
+
+					const float OneMinusInterp = 1.0 - InterpolationFactor;
+					const int32 InterpFixed = (int32)(InterpolationFactor * 255.0f);
+					const int32 OneMinusInterpFixed = 255 - InterpFixed;
+
+					for (int32 Index = 0; Index < NumVerts; ++Index)
+					{
+						const FVector& PositionA = TrackProxy->MeshData->Positions[Index];
+						const FVector& PositionB = TrackProxy->NextFrameMeshData->Positions[Index];
+						InterpolatedPositions[Index] = PositionA * OneMinusInterp + PositionB* InterpolationFactor;
+					}
+
+					for (int32 Index = 0; Index < NumVerts; ++Index)
+					{
+						// The following are already 8 bit so quantized enough we can do exact equal comparisons
+						const FPackedNormal& TangentXA = TrackProxy->MeshData->TangentsX[Index];
+						const FPackedNormal& TangentXB = TrackProxy->NextFrameMeshData->TangentsX[Index];
+						const FPackedNormal& TangentZA = TrackProxy->MeshData->TangentsZ[Index];
+						const FPackedNormal& TangentZB = TrackProxy->NextFrameMeshData->TangentsZ[Index];
+
+						InterpolatedTangentX[Index] = InterpolatePackedNormal(TangentXA, TangentXB, InterpFixed, OneMinusInterpFixed);
+						InterpolatedTangentZ[Index] = InterpolatePackedNormal(TangentZA, TangentZB, InterpFixed, OneMinusInterpFixed);
+					}
+
+					if (TrackProxy->MeshData->VertexInfo.bHasColor0) for (int32 Index = 0; Index < NumVerts; ++Index)
+					{
+						const FColor& ColorA = TrackProxy->MeshData->Colors[Index];
+						const FColor& ColorB = TrackProxy->NextFrameMeshData->Colors[Index];
+						InterpolatedColors[Index] = InterpolatePackedColor(ColorA, ColorB, InterpFixed, OneMinusInterpFixed);
+					}
+
+					if (TrackProxy->MeshData->VertexInfo.bHasUV0) for (int32 Index = 0; Index < NumVerts; ++Index)
+					{
+						const FVector2D& UVA = TrackProxy->MeshData->TextureCoordinates[Index];
+						const FVector2D& UVB = TrackProxy->NextFrameMeshData->TextureCoordinates[Index];
+
+						InterpolatedUVs[Index] = UVA * OneMinusInterp + UVB * InterpolationFactor;
+					}
+
+					if (bHasMotionVectors) for (int32 Index = 0; Index < NumVerts; ++Index)
+					{
+						InterpolatedMotionVectors[Index] = TrackProxy->MeshData->MotionVectors[Index] * OneMinusInterp +
+							TrackProxy->NextFrameMeshData->MotionVectors[Index] * InterpolationFactor;
+					}
+
+					// Upload other non-motionblurred data
+					if (!TrackProxy->MeshData->VertexInfo.bConstantIndices)
+						TrackProxy->IndexBuffer.Update(TrackProxy->MeshData->Indices);
+
+					if (TrackProxy->MeshData->VertexInfo.bHasTangentX)
+						TrackProxy->TangentXBuffer.Update(InterpolatedTangentX);
+					if (TrackProxy->MeshData->VertexInfo.bHasTangentZ)
+						TrackProxy->TangentZBuffer.Update(InterpolatedTangentZ);
+
+					if (TrackProxy->MeshData->VertexInfo.bHasUV0)
+						TrackProxy->TextureCoordinatesBuffer.Update(InterpolatedUVs);
+
+					if (TrackProxy->MeshData->VertexInfo.bHasColor0)
+						TrackProxy->ColorBuffer.Update(InterpolatedColors);
+
+					bool bIsCompatibleWithCachedFrame = TrackProxy->Resource->IsTopologyCompatible(
+						TrackProxy->PositionBufferFrameIndices[TrackProxy->CurrentPositionBufferIndex % 2],
+						TrackProxy->FrameIndex);
+
+					if (!bHasMotionVectors)
+					{
+						// Initialize both buffers the first frame
+						if (TrackProxy->CurrentPositionBufferIndex == -1 || !bIsCompatibleWithCachedFrame)
+						{
+							TrackProxy->PositionBuffers[0].Update(InterpolatedPositions);
+							TrackProxy->PositionBuffers[1].Update(InterpolatedPositions);
+							TrackProxy->CurrentPositionBufferIndex = 0;
+							TrackProxy->PositionBufferFrameTimes[0] = Time;
+							TrackProxy->PositionBufferFrameTimes[1] = Time;
+							// We need to keep a frame index in order to ensure topology consistency. As we can interpolate 
+							// FrameIndex and NextFrameIndex are certainly topo-compatible so it doesn't really matter which 
+							// one we keep here. But wee keep NextFrameIndex as that is most useful to validate against
+							// the frame coming up
+							TrackProxy->PositionBufferFrameIndices[0] = TrackProxy->NextFrameIndex;
+							TrackProxy->PositionBufferFrameIndices[1] = TrackProxy->NextFrameIndex;
+						}
+						else
+						{
+							TrackProxy->CurrentPositionBufferIndex++;
+							TrackProxy->PositionBuffers[TrackProxy->CurrentPositionBufferIndex % 2].Update(InterpolatedPositions);
+							TrackProxy->PositionBufferFrameTimes[TrackProxy->CurrentPositionBufferIndex % 2] = Time;
+							TrackProxy->PositionBufferFrameIndices[TrackProxy->CurrentPositionBufferIndex % 2] = TrackProxy->NextFrameIndex;
+						}
+					}
+					else
+					{
+						TrackProxy->CurrentPositionBufferIndex = 0;
+						TrackProxy->PositionBuffers[0].Update(InterpolatedPositions);
+						TrackProxy->PositionBuffers[1].Update(InterpolatedMotionVectors);
+						TrackProxy->PositionBufferFrameIndices[0] = TrackProxy->FrameIndex;
+						TrackProxy->PositionBufferFrameIndices[1] = -1;
+						TrackProxy->PositionBufferFrameTimes[0] = Time;
+						TrackProxy->PositionBufferFrameTimes[1] = Time;
+					}
+				}
+
+
+			}
+			else
+			{
+				// We just don't interpolate between frames if we got GPU to burn we could someday render twice and stipple fade between it :-D like with lods
+
+				// Only bother uploading if anything changed or when the we failed to decode anything make sure update the gpu buffers regardless
+				if (bDecodedAnything || bDecoderError)
+				{
+					//TrackProxy->IndexBuffer.Update(TrackProxy->MeshData->Indices);
+
+					const int32 NumVertices = TrackProxy->MeshData->Positions.Num();
+
+					if (TrackProxy->MeshData->VertexInfo.bHasTangentX)
+						TrackProxy->TangentXBuffer.Update(TrackProxy->MeshData->TangentsX);
+					if (TrackProxy->MeshData->VertexInfo.bHasTangentZ)
+						TrackProxy->TangentZBuffer.Update(TrackProxy->MeshData->TangentsZ);
+
+					if (!TrackProxy->MeshData->VertexInfo.bConstantIndices)
+						TrackProxy->IndexBuffer.Update(TrackProxy->MeshData->Indices);
+
+					if (TrackProxy->MeshData->VertexInfo.bHasUV0)
+						TrackProxy->TextureCoordinatesBuffer.Update(TrackProxy->MeshData->TextureCoordinates);
+
+					if (TrackProxy->MeshData->VertexInfo.bHasColor0)
+						TrackProxy->ColorBuffer.Update(TrackProxy->MeshData->Colors);
+
+					bool bIsCompatibleWithCachedFrame = TrackProxy->Resource->IsTopologyCompatible(
+						TrackProxy->PositionBufferFrameIndices[TrackProxy->CurrentPositionBufferIndex % 2],
+						TrackProxy->FrameIndex);
+
+					if (!bHasMotionVectors)
+					{
+						// Initialize both buffers the first frame or when topology changed as we can't render
+						// with a previous buffer referencing a buffer from another topology
+						if (TrackProxy->CurrentPositionBufferIndex == -1 || !bIsCompatibleWithCachedFrame || bSeeked)
+						{
+							TrackProxy->PositionBuffers[0].Update(TrackProxy->MeshData->Positions);
+							TrackProxy->PositionBuffers[1].Update(TrackProxy->MeshData->Positions);
+							TrackProxy->CurrentPositionBufferIndex = 0;
+							TrackProxy->PositionBufferFrameIndices[0] = TrackProxy->FrameIndex;
+							TrackProxy->PositionBufferFrameIndices[1] = TrackProxy->FrameIndex;
+						}
+						// We still use the previous frame's buffer as a motion blur previous position. As interpolation is switched
+						// off the actual time of this previous frame depends on the geometry cache framerate and playback speed
+						// so the motion blur vectors may not really be anything relevant. Do we want to just disable motion blur? 
+						// But as an optimization skipping interpolation when the cache fps is near to the actual game fps this is obviously nice...
+						else
+						{
+							TrackProxy->CurrentPositionBufferIndex++;
+							TrackProxy->PositionBuffers[TrackProxy->CurrentPositionBufferIndex % 2].Update(TrackProxy->MeshData->Positions);
+							TrackProxy->PositionBufferFrameIndices[TrackProxy->CurrentPositionBufferIndex % 2] = TrackProxy->FrameIndex;
+						}
+					}
+					else
+					{
+						TrackProxy->CurrentPositionBufferIndex = 0;
+						TrackProxy->PositionBuffers[0].Update(TrackProxy->MeshData->Positions);
+						TrackProxy->PositionBuffers[1].Update(TrackProxy->MeshData->MotionVectors);
+						TrackProxy->PositionBufferFrameIndices[0] = TrackProxy->FrameIndex;
+						TrackProxy->PositionBufferFrameIndices[1] = -1;
+						TrackProxy->PositionBufferFrameTimes[0] = Time;
+						TrackProxy->PositionBufferFrameTimes[1] = Time;
+					}
+				}
+			}
+
+		#if 0
+			bool bOffloadUpdate = CVarOffloadUpdate.GetValueOnRenderThread() != 0;
+			if (TrackProxy->SampleIndex != TrackProxy->UploadedSampleIndex)
+			{
+				TrackProxy->UploadedSampleIndex = TrackProxy->SampleIndex;
+
+				if (bOffloadUpdate)
+				{
+					check(false);
+					// Only update the size on this thread
+					TrackProxy->IndexBuffer.UpdateSizeOnly(TrackProxy->MeshData->Indices.Num());
+					TrackProxy->VertexBuffer.UpdateSizeTyped<FNoPositionVertex>(TrackProxy->MeshData->Vertices.Num());
+
+					// Do the interpolation on a worker thread
+					FGraphEventRef CompletionFence = FFunctionGraphTask::CreateAndDispatchWhenReady([]()
+					{
+
+					}, GET_STATID(STAT_BufferUpdateTask), NULL, ENamedThreads::AnyThread);
+
+					// Queue a command on the RHI thread that waits for the interpolation job and then uploads them to the GPU
+					FRHICommandListImmediate &RHICommandList = GetImmediateCommandList_ForRenderCommand();
+					new (RHICommandList.AllocCommand<FRHICommandUpdateGeometryCacheBuffer>())FRHICommandUpdateGeometryCacheBuffer(
+						CompletionFence,
+						TrackProxy->VertexBuffer.VertexBufferRHI,
+						TrackProxy->MeshData->Vertices.GetData(),
+						TrackProxy->VertexBuffer.GetSizeInBytes(),
+						TrackProxy->IndexBuffer.IndexBufferRHI,
+						TrackProxy->MeshData->Indices.GetData(),
+						TrackProxy->IndexBuffer.SizeInBytes());
+
+					// Upload vertex data
+					/*void* RESTRICT Data = (void* RESTRICT)GDynamicRHI->RHILockVertexBuffer(TrackProxy->VertexBuffer.VertexBufferRHI, 0, TrackProxy->VertexBuffer.SizeInBytes(), RLM_WriteOnly);
+					FMemory::BigBlockMemcpy(Data, TrackProxy->MeshData->Vertices.GetData(), TrackProxy->VertexBuffer.SizeInBytes());
+					GDynamicRHI->RHIUnlockVertexBuffer(TrackProxy->VertexBuffer.VertexBufferRHI);
+
+					// Upload index data
+					Data = (void* RESTRICT)GDynamicRHI->RHILockIndexBuffer(TrackProxy->IndexBuffer.IndexBufferRHI, 0, TrackProxy->IndexBuffer.SizeInBytes(), RLM_WriteOnly);
+					FMemory::BigBlockMemcpy(Data, TrackProxy->MeshData->Indices.GetData(), TrackProxy->IndexBuffer.SizeInBytes());
+					GDynamicRHI->RHIUnlockIndexBuffer(TrackProxy->IndexBuffer.IndexBufferRHI);*/
+				}
+				else
+				{
+					TrackProxy->IndexBuffer.Update(TrackProxy->MeshData->Indices);
+					TrackProxy->VertexBuffer.Update(TrackProxy->MeshData->Vertices);
+
+					// Initialize both buffers the first frame
+					if (TrackProxy->CurrentPositionBufferIndex == -1)
+					{
+						TrackProxy->PositonBuffers[0].Update(TrackProxy->MeshData->Vertices);
+						TrackProxy->PositonBuffers[1].Update(TrackProxy->MeshData->Vertices);
+						TrackProxy->CurrentPositionBufferIndex = 0;
+					}
+					else
+					{
+						TrackProxy->CurrentPositionBufferIndex++;
+						TrackProxy->PositonBuffers[TrackProxy->CurrentPositionBufferIndex%2].Update(TrackProxy->MeshData->Vertices);
+					}
+				}
+			}
+		#endif
+
+		}
+	}
 }
 
 void FGeometryCacheSceneProxy::UpdateSectionWorldMatrix(const int32 SectionIndex, const FMatrix& WorldMatrix)
@@ -822,7 +932,7 @@ void FGeomCacheIndexBuffer::InitRHI()
 {
 	FRHIResourceCreateInfo CreateInfo;
 	void* Buffer = nullptr;
-	IndexBufferRHI = RHICreateAndLockIndexBuffer(sizeof(uint32), NumIndices * sizeof(uint32), BUF_Static, CreateInfo, Buffer);
+	IndexBufferRHI = RHICreateAndLockIndexBuffer(sizeof(uint32), NumIndices * sizeof(uint32), BUF_Static | BUF_ShaderResource, CreateInfo, Buffer);
 	RHIUnlockIndexBuffer(IndexBufferRHI);
 }
 
@@ -839,7 +949,7 @@ void FGeomCacheIndexBuffer::Update(const TArray<uint32> &Indices)
 	{
 		NumIndices = Indices.Num();
 		FRHIResourceCreateInfo CreateInfo;
-		IndexBufferRHI = RHICreateAndLockIndexBuffer(sizeof(uint32), NumIndices * sizeof(uint32), BUF_Static, CreateInfo, Buffer);		
+		IndexBufferRHI = RHICreateAndLockIndexBuffer(sizeof(uint32), NumIndices * sizeof(uint32), BUF_Static | BUF_ShaderResource, CreateInfo, Buffer);
 	}
 	else
 	{
@@ -859,7 +969,7 @@ void FGeomCacheIndexBuffer::UpdateSizeOnly(int32 NewNumIndices)
 	if (NewNumIndices > NumIndices)
 	{
 		FRHIResourceCreateInfo CreateInfo;
-		IndexBufferRHI = RHICreateIndexBuffer(sizeof(uint32), NewNumIndices * sizeof(uint32), BUF_Static, CreateInfo);
+		IndexBufferRHI = RHICreateIndexBuffer(sizeof(uint32), NewNumIndices * sizeof(uint32), BUF_Static | BUF_ShaderResource, CreateInfo);
 		NumIndices = NewNumIndices;
 	}
 }
@@ -868,7 +978,7 @@ void FGeomCacheVertexBuffer::InitRHI()
 {
 	FRHIResourceCreateInfo CreateInfo;
 	void* BufferData = nullptr;
-	VertexBufferRHI = RHICreateAndLockVertexBuffer(SizeInBytes, BUF_Static, CreateInfo, BufferData);
+	VertexBufferRHI = RHICreateAndLockVertexBuffer(SizeInBytes, BUF_Static | BUF_ShaderResource, CreateInfo, BufferData);
 	RHIUnlockVertexBuffer(VertexBufferRHI);
 }
 
@@ -884,7 +994,7 @@ void FGeomCacheVertexBuffer::UpdateRaw(const void *Data, int32 NumItems, int32 I
 	{
 		SizeInBytes = NewSizeInBytes;
 		FRHIResourceCreateInfo CreateInfo;
-		VertexBufferRHI = RHICreateAndLockVertexBuffer(SizeInBytes, BUF_Static, CreateInfo, VertexBufferData);
+		VertexBufferRHI = RHICreateAndLockVertexBuffer(SizeInBytes, BUF_Static | BUF_ShaderResource, CreateInfo, VertexBufferData);
 	}
 	else
 	{
@@ -916,6 +1026,6 @@ void FGeomCacheVertexBuffer::UpdateSize(int32 NewSizeInBytes)
 	{
 		SizeInBytes = NewSizeInBytes;
 		FRHIResourceCreateInfo CreateInfo;
-		VertexBufferRHI = RHICreateVertexBuffer(SizeInBytes, BUF_Static, CreateInfo);
+		VertexBufferRHI = RHICreateVertexBuffer(SizeInBytes, BUF_Static | BUF_ShaderResource, CreateInfo);
 	}
 }
diff --git a/Engine/Plugins/Experimental/GeometryCache/Source/GeometryCache/Private/GeometryCacheSceneProxy.h b/Engine/Plugins/Experimental/GeometryCache/Source/GeometryCache/Private/GeometryCacheSceneProxy.h
--- a/Engine/Plugins/Experimental/GeometryCache/Source/GeometryCache/Private/GeometryCacheSceneProxy.h
+++ b/Engine/Plugins/Experimental/GeometryCache/Source/GeometryCache/Private/GeometryCacheSceneProxy.h
@@ -168,6 +168,10 @@ public:
 
 	FGeometryCacheTrackStreamableRenderResource *Resource;
 	int32 UploadedSampleIndex;
+
+#if RHI_RAYTRACING
+	FRayTracingGeometry RayTracingGeometry;
+#endif
 };
 
 /** Procedural mesh scene proxy */
@@ -201,6 +205,11 @@ public:
 
 	/** Clears the Sections array*/
 	void ClearSections();
+
+#if RHI_RAYTRACING
+	virtual void GetRayTracingGeometryInstances(TArray<FRayTracingGeometryInstanceCollection>& OutInstanceCollections) override;
+	virtual bool IsRayTracingRelevant() const override { return true; }
+#endif
 	
 private:	
 	FMaterialRelevance MaterialRelevance;
@@ -211,4 +220,6 @@ private:
 
 	/** Array of Track Proxies */
 	TArray<FGeomCacheTrackProxy*> Tracks;
+
+	void FrameUpdate() const;
 };
