Fix Mac (non-raytracing) build.

diff --git a/Engine/Plugins/Experimental/GeometryCache/Source/GeometryCache/Private/GeometryCacheSceneProxy.cpp b/Engine/Plugins/Experimental/GeometryCache/Source/GeometryCache/Private/GeometryCacheSceneProxy.cpp
--- a/Engine/Plugins/Experimental/GeometryCache/Source/GeometryCache/Private/GeometryCacheSceneProxy.cpp
+++ b/Engine/Plugins/Experimental/GeometryCache/Source/GeometryCache/Private/GeometryCacheSceneProxy.cpp
@@ -440,6 +440,7 @@ void FGeometryCacheSceneProxy::GetDynamicMeshElements(const TArray<const FSceneV
 	}
 }
 
+#if RHI_RAYTRACING
 void FGeometryCacheSceneProxy::GetRayTracingGeometryInstances(TArray<FRayTracingGeometryInstanceCollection>& OutInstanceCollections)
 {
 	for (FGeomCacheTrackProxy* TrackProxy : Tracks)
@@ -458,6 +459,7 @@ void FGeometryCacheSceneProxy::GetRayTracingGeometryInstances(TArray<FRayTracing
 		OutInstanceCollections.Add(Collection);
 	}
 }
+#endif
 
 FPrimitiveViewRelevance FGeometryCacheSceneProxy::GetViewRelevance(const FSceneView* View) const
 {
