Fixed static analysis warning in Win64 (PVS-Studio)

diff --git a/Engine/Plugins/Experimental/GeometryCache/Source/GeometryCache/Private/CodecV1.cpp b/Engine/Plugins/Experimental/GeometryCache/Source/GeometryCache/Private/CodecV1.cpp
--- a/Engine/Plugins/Experimental/GeometryCache/Source/GeometryCache/Private/CodecV1.cpp
+++ b/Engine/Plugins/Experimental/GeometryCache/Source/GeometryCache/Private/CodecV1.cpp
@@ -721,7 +721,7 @@ void FCodecV1Encoder::WriteInt32(FHuffmanEncodeTable& ValueTable, int32 Value)
 			int32 NumRawBits = HighestSetBit(~Value);	// Find first 0-bit. 1 <= NumRawBits <= 30.
 			int32 Packed = 3 + NumRawBits * 2;			// First negative code is 5
 			WriteSymbol(ValueTable, Packed);
-			int32 RawBits = Value & ~(-1 << NumRawBits);
+			int32 RawBits = Value & ~(0xFFFFFFFFu << NumRawBits);
 			WriteBits(RawBits, NumRawBits);
 		}
 	}
