Fix for geometry caches evaluating in packaged builds and use UGeometryCache as an asset, adding needed functions to it.

(Will talk to Jurre about maybe putting Duration itself on the UGeometryCache itself, as it is there is some duplication since I don't want the Component to call the slightly inefficient UGeometryCache implementation that manually calc's the duration isntead of caching it).

Note that I changed GeometryCache to a weak pointer instead of a soft object path recently to try to fix the spawnable issue.  So the conversion is for older files is valid.

diff --git a/Engine/Plugins/Experimental/GeometryCache/Source/GeometryCache/Classes/GeometryCache.h b/Engine/Plugins/Experimental/GeometryCache/Source/GeometryCache/Classes/GeometryCache.h
--- a/Engine/Plugins/Experimental/GeometryCache/Source/GeometryCache/Classes/GeometryCache.h
+++ b/Engine/Plugins/Experimental/GeometryCache/Source/GeometryCache/Classes/GeometryCache.h
@@ -71,6 +71,13 @@ public:
 
 	/** Get the end frame */
 	int32 GetEndFrame() const;
+	
+	/** Calculate it's duration */
+	float CalculateDuration() const;
+
+	/** Get the Frame at the Specified Time*/
+	int32  GetFrameAtTime(const float Time) const;
+
 private:
 	/** A fence which is used to keep track of the rendering thread releasing the geometry cache resources. */
 	FRenderCommandFence ReleaseResourcesFence;
diff --git a/Engine/Plugins/Experimental/GeometryCache/Source/GeometryCache/Private/GeometryCache.cpp b/Engine/Plugins/Experimental/GeometryCache/Source/GeometryCache/Private/GeometryCache.cpp
--- a/Engine/Plugins/Experimental/GeometryCache/Source/GeometryCache/Private/GeometryCache.cpp
+++ b/Engine/Plugins/Experimental/GeometryCache/Source/GeometryCache/Private/GeometryCache.cpp
@@ -156,4 +156,27 @@ int32 UGeometryCache::GetEndFrame() const
 	return EndFrame;
 }
 
+float UGeometryCache::CalculateDuration() const
+{
+	int32 NumTracks = Tracks.Num();
+	float Duration = 0.0f;
+	// Create mesh sections for each GeometryCacheTrack
+	for (int32 TrackIndex = 0; TrackIndex < NumTracks; ++TrackIndex)
+	{
+		const float TrackMaxSampleTime = Tracks[TrackIndex]->GetMaxSampleTime();
+		Duration = (Duration > TrackMaxSampleTime) ? Duration : TrackMaxSampleTime;
+	}
+	return Duration;
+}
+
+int32 UGeometryCache::GetFrameAtTime(const float Time) const
+{
+	const float Duration = CalculateDuration();
+	const int32 NumberOfFrames = GetEndFrame() - GetStartFrame() + 1;;
+	const float FrameTime = NumberOfFrames > 1 ? Duration / (float)(NumberOfFrames - 1) : 0.0f;
+	const int32 NormalizedFrame = FMath::Clamp(FMath::RoundToInt(Time / FrameTime), 0, NumberOfFrames - 1);
+	const int32 StartFrame =  GetStartFrame();
+	return StartFrame + NormalizedFrame; 
+
+}
 #undef LOCTEXT_NAMESPACE // "GeometryCache"
\ No newline at end of file
diff --git a/Engine/Plugins/Experimental/GeometryCache/Source/GeometryCacheSequencer/Classes/GeometryCacheTrackEditor.cpp b/Engine/Plugins/Experimental/GeometryCache/Source/GeometryCacheSequencer/Classes/GeometryCacheTrackEditor.cpp
--- a/Engine/Plugins/Experimental/GeometryCache/Source/GeometryCacheSequencer/Classes/GeometryCacheTrackEditor.cpp
+++ b/Engine/Plugins/Experimental/GeometryCache/Source/GeometryCacheSequencer/Classes/GeometryCacheTrackEditor.cpp
@@ -80,13 +80,10 @@ UMovieSceneSection* FGeometryCacheSection::GetSectionObject()
 
 FText FGeometryCacheSection::GetSectionTitle() const
 {
-	if (Section.Params.GeometryCache.Get() != nullptr)
+	if (Section.Params.GeometryCacheAsset != nullptr)
 	{
-		UGeometryCacheComponent* GeometryCache = Cast<UGeometryCacheComponent>(Section.Params.GeometryCache.Get());
-		if (GeometryCache && GeometryCache->GetOwner() != nullptr)
-		{
-			return FText::FromString(GeometryCache->GetOwner()->GetName());
-		}
+		return FText::FromString(Section.Params.GeometryCacheAsset->GetName());
+	
 	}
 	return LOCTEXT("NoGeometryCacheSection", "No GeometryCache");
 }
@@ -117,7 +114,8 @@ int32 FGeometryCacheSection::OnPaintSection(FSequencerSectionPainter& Painter) c
 
 	// Add lines where the animation starts and ends/loops
 	float AnimPlayRate = FMath::IsNearlyZero(Section.Params.PlayRate) ? 1.0f : Section.Params.PlayRate;
-	float SeqLength = Section.Params.GetSequenceLength() - TickResolution.AsSeconds(Section.Params.StartFrameOffset + Section.Params.EndFrameOffset) / AnimPlayRate;
+	float Duration = Section.Params.GetSequenceLength();
+	float SeqLength = Duration - TickResolution.AsSeconds(Section.Params.StartFrameOffset + Section.Params.EndFrameOffset) / AnimPlayRate;
 
 	if (!FMath::IsNearlyZero(SeqLength, KINDA_SMALL_NUMBER) && SeqLength > 0)
 	{
@@ -148,14 +146,14 @@ int32 FGeometryCacheSection::OnPaintSection(FSequencerSectionPainter& Painter) c
 	if (Painter.bIsSelected && SequencerPtr.IsValid())
 	{
 		FFrameTime CurrentTime = SequencerPtr->GetLocalTime().Time;
-		if (Section.GetRange().Contains(CurrentTime.FrameNumber) && Section.Params.GeometryCache.Get() != nullptr)
+		if (Section.GetRange().Contains(CurrentTime.FrameNumber) && Section.Params.GeometryCacheAsset != nullptr)
 		{
 			const float Time = TimeToPixelConverter.FrameToPixel(CurrentTime);
 
-			UGeometryCacheComponent* GeometryCache = Cast<UGeometryCacheComponent>(Section.Params.GeometryCache.Get());
+			UGeometryCache* GeometryCache = Section.Params.GeometryCacheAsset;
 
 			// Draw the current time next to the scrub handle
-			const float AnimTime = Section.MapTimeToAnimation(CurrentTime, TickResolution);
+			const float AnimTime = Section.MapTimeToAnimation(Duration, CurrentTime, TickResolution);
 			int32 FrameTime = GeometryCache->GetFrameAtTime(AnimTime);
 			FString FrameString = FString::FromInt(FrameTime);
 
diff --git a/Engine/Plugins/Experimental/GeometryCache/Source/GeometryCacheTracks/Classes/MovieSceneGeometryCacheSection.h b/Engine/Plugins/Experimental/GeometryCache/Source/GeometryCacheTracks/Classes/MovieSceneGeometryCacheSection.h
--- a/Engine/Plugins/Experimental/GeometryCache/Source/GeometryCacheTracks/Classes/MovieSceneGeometryCacheSection.h
+++ b/Engine/Plugins/Experimental/GeometryCache/Source/GeometryCacheTracks/Classes/MovieSceneGeometryCacheSection.h
@@ -6,22 +6,23 @@
 #include "UObject/ObjectMacros.h"
 #include "MovieSceneSection.h"
 #include "GeometryCacheComponent.h"
+#include "GeometryCache.h"
 #include "Channels/MovieSceneFloatChannel.h"
 #include "UObject/SoftObjectPath.h"
 #include "MovieSceneGeometryCacheSection.generated.h"
 
 USTRUCT()
-struct FMovieSceneGeometryCacheParams
+ struct GEOMETRYCACHETRACKS_API FMovieSceneGeometryCacheParams
 {
 	GENERATED_BODY()
 
 	FMovieSceneGeometryCacheParams();
 
-	/** Gets the animation duration, modified by play rate */
-	float GetDuration() const { return FMath::IsNearlyZero(PlayRate) || GeometryCache.Get() == nullptr ? 0.f : Cast<UGeometryCacheComponent>(GeometryCache.Get())->GetDuration() / PlayRate; }
-
 	/** Gets the animation sequence length, not modified by play rate */
-	float GetSequenceLength() const { return GeometryCache.Get() != nullptr ? Cast<UGeometryCacheComponent>(GeometryCache.Get())->GetDuration() : 0.f; }
+	 float GetSequenceLength() const;
+	/** The animation this section plays */
+	UPROPERTY(EditAnywhere, Category = "GeometryCache", DisplayName = "Geometry Cache")
+	UGeometryCache* GeometryCacheAsset;
 
 	/** The offset into the beginning of the animation clip */
 	UPROPERTY(EditAnywhere, Category = "GeometryCache")
@@ -45,7 +46,8 @@ struct FMovieSceneGeometryCacheParams
 	UPROPERTY()
 	float EndOffset_DEPRECATED;
 
-	TWeakObjectPtr<UGeometryCacheComponent> GeometryCache;
+	UPROPERTY()
+	FSoftObjectPath GeometryCache_DEPRECATED;
 };
 
 /**
@@ -63,7 +65,7 @@ public:
 	FMovieSceneGeometryCacheParams Params;
 
 	/** Get Frame Time as Animation Time*/
-	virtual float MapTimeToAnimation(FFrameTime InPosition, FFrameRate InFrameRate) const;
+	virtual float MapTimeToAnimation(float ComponentDuration, FFrameTime InPosition, FFrameRate InFrameRate) const;
 
 protected:
 	//~ UMovieSceneSection interface
diff --git a/Engine/Plugins/Experimental/GeometryCache/Source/GeometryCacheTracks/Private/MovieSceneGeometryCacheSection.cpp b/Engine/Plugins/Experimental/GeometryCache/Source/GeometryCacheTracks/Private/MovieSceneGeometryCacheSection.cpp
--- a/Engine/Plugins/Experimental/GeometryCache/Source/GeometryCacheTracks/Private/MovieSceneGeometryCacheSection.cpp
+++ b/Engine/Plugins/Experimental/GeometryCache/Source/GeometryCacheTracks/Private/MovieSceneGeometryCacheSection.cpp
@@ -16,6 +16,8 @@ namespace
 
 FMovieSceneGeometryCacheParams::FMovieSceneGeometryCacheParams()
 {
+	GeometryCacheAsset = nullptr;
+	GeometryCache_DEPRECATED = nullptr;
 	StartOffset_DEPRECATED = GeometryCacheDeprecatedMagicNumber;
 	EndOffset_DEPRECATED = GeometryCacheDeprecatedMagicNumber;
 	PlayRate = 1.f;
@@ -59,6 +61,16 @@ void UMovieSceneGeometryCacheSection::PostLoad()
 
 		Params.EndOffset_DEPRECATED = GeometryCacheDeprecatedMagicNumber;
 	}
+
+	if (Params.GeometryCache_DEPRECATED.ResolveObject() != nullptr
+		&& Params.GeometryCacheAsset == nullptr)
+	{
+		UGeometryCacheComponent *Comp = Cast<UGeometryCacheComponent>(Params.GeometryCache_DEPRECATED.ResolveObject());
+		if (Comp)
+		{
+			Params.GeometryCacheAsset = (Comp->GetGeometryCache());
+		}
+	}
 }
 
 void UMovieSceneGeometryCacheSection::Serialize(FArchive& Ar)
@@ -153,10 +165,10 @@ void UMovieSceneGeometryCacheSection::GetSnapTimes(TArray<FFrameNumber>& OutSnap
 	}
 }
 
-float UMovieSceneGeometryCacheSection::MapTimeToAnimation(FFrameTime InPosition, FFrameRate InFrameRate) const
+float UMovieSceneGeometryCacheSection::MapTimeToAnimation(float ComponentDuration, FFrameTime InPosition, FFrameRate InFrameRate) const
 {
-	FMovieSceneGeometryCacheSectionTemplateParameters TemplateParams(const_cast<FMovieSceneGeometryCacheParams&> (Params), GetInclusiveStartFrame(), GetExclusiveEndFrame());
-	return TemplateParams.MapTimeToAnimation(InPosition, InFrameRate);
+	FMovieSceneGeometryCacheSectionTemplateParameters TemplateParams(Params, GetInclusiveStartFrame(), GetExclusiveEndFrame());
+	return TemplateParams.MapTimeToAnimation(ComponentDuration, InPosition, InFrameRate);
 }
 
 
@@ -186,9 +198,14 @@ void UMovieSceneGeometryCacheSection::PostEditChangeProperty(FPropertyChangedEve
 			PreviousPlayRate = NewPlayRate;
 		}
 	}
-
+	
 	Super::PostEditChangeProperty(PropertyChangedEvent);
 }
+
+float FMovieSceneGeometryCacheParams::GetSequenceLength() const
+{
+	return GeometryCacheAsset != nullptr ? GeometryCacheAsset->CalculateDuration() : 0.f;
+}
 #endif
 
 #undef LOCTEXT_NAMESPACE 
diff --git a/Engine/Plugins/Experimental/GeometryCache/Source/GeometryCacheTracks/Private/MovieSceneGeometryCacheTemplate.cpp b/Engine/Plugins/Experimental/GeometryCache/Source/GeometryCacheTracks/Private/MovieSceneGeometryCacheTemplate.cpp
--- a/Engine/Plugins/Experimental/GeometryCache/Source/GeometryCacheTracks/Private/MovieSceneGeometryCacheTemplate.cpp
+++ b/Engine/Plugins/Experimental/GeometryCache/Source/GeometryCacheTracks/Private/MovieSceneGeometryCacheTemplate.cpp
@@ -83,18 +83,21 @@ struct FGeometryCacheExecutionToken
 			{
 				if (UObject* Obj = WeakObj.Get())
 				{
-					UGeometryCacheComponent* GeometryCache = GeometryMeshComponentFromObject(Obj);
-					if (GeometryCache)
+					UGeometryCacheComponent* GeometryComp = GeometryMeshComponentFromObject(Obj);
+					if (GeometryComp)
 					{
-						if (Params.SectionParams->GeometryCache.Get() == nullptr)
+						if (Params.GeometryCacheAsset != GeometryComp->GetGeometryCache())
 						{
-							Params.SectionParams->GeometryCache = GeometryCache;
+							UGeometryCache* GeomCache = Params.GeometryCacheAsset;
+							{
+								GeometryComp->SetGeometryCache(GeomCache);
+							}
 						}
-						Player.SavePreAnimatedState(*GeometryCache, FPreAnimatedGeometryCacheTokenProducer::GetAnimTypeID(), FPreAnimatedGeometryCacheTokenProducer());
-						GeometryCache->SetManualTick(true);
+						Player.SavePreAnimatedState(*GeometryComp, FPreAnimatedGeometryCacheTokenProducer::GetAnimTypeID(), FPreAnimatedGeometryCacheTokenProducer());
+						GeometryComp->SetManualTick(true);
 						// calculate the time at which to evaluate the animation
-						float EvalTime = Params.MapTimeToAnimation(Context.GetTime(), Context.GetFrameRate());
-						GeometryCache->TickAtThisTime(EvalTime, true, Params.SectionParams->bReverse, true);
+						float EvalTime = Params.MapTimeToAnimation(GeometryComp->GetDuration(), Context.GetTime(), Context.GetFrameRate());
+						GeometryComp->TickAtThisTime(EvalTime, true, Params.bReverse, true);
 					}
 				}
 			}
@@ -117,31 +120,31 @@ void FMovieSceneGeometryCacheSectionTemplate::Evaluate(const FMovieSceneEvaluati
 		ExecutionTokens.Add(FGeometryCacheExecutionToken(Params));
 }
 
-float FMovieSceneGeometryCacheSectionTemplateParameters::MapTimeToAnimation(FFrameTime InPosition, FFrameRate InFrameRate) const
+float FMovieSceneGeometryCacheSectionTemplateParameters::MapTimeToAnimation(float ComponentDuration, FFrameTime InPosition, FFrameRate InFrameRate) const
 {
-	FFrameTime AnimationLength = SectionParams->GetSequenceLength() * InFrameRate;
+	float SequenceLength = ComponentDuration;
+	FFrameTime AnimationLength = SequenceLength * InFrameRate;
 	int32 LengthInFrames = AnimationLength.FrameNumber.Value + (int)(AnimationLength.GetSubFrame() + 0.5f) + 1;
 	//we only play end if we are not looping, and assuming we are looping if Length is greater than default length;
 	bool bLooping = (SectionEndTime.Value - SectionStartTime.Value) > LengthInFrames;
 
 	InPosition = FMath::Clamp(InPosition, FFrameTime(SectionStartTime), FFrameTime(SectionEndTime - 1));
 
-	const float SectionPlayRate = SectionParams->PlayRate;
+	const float SectionPlayRate = PlayRate;
 	const float AnimPlayRate = FMath::IsNearlyZero(SectionPlayRate) ? 1.0f : SectionPlayRate;
 
-	const float SeqLength = SectionParams->GetSequenceLength() - InFrameRate.AsSeconds(SectionParams->StartFrameOffset + SectionParams->EndFrameOffset);
+	const float SeqLength = SequenceLength - InFrameRate.AsSeconds(StartFrameOffset + EndFrameOffset);
 
 	float AnimPosition = FFrameTime::FromDecimal((InPosition - SectionStartTime).AsDecimal() * AnimPlayRate) / InFrameRate;
 	if (SeqLength > 0.f && (bLooping || !FMath::IsNearlyEqual(AnimPosition, SeqLength, 1e-4f)))
 	{
 		AnimPosition = FMath::Fmod(AnimPosition, SeqLength);
 	}
-	AnimPosition += InFrameRate.AsSeconds(SectionParams->StartFrameOffset);
-	if (SectionParams->bReverse)
+	AnimPosition += InFrameRate.AsSeconds(StartFrameOffset);
+	if (bReverse)
 	{
-		AnimPosition = (SeqLength - (AnimPosition - InFrameRate.AsSeconds(SectionParams->StartFrameOffset))) + InFrameRate.AsSeconds(SectionParams->StartFrameOffset);
+		AnimPosition = (SeqLength - (AnimPosition - InFrameRate.AsSeconds(StartFrameOffset))) + InFrameRate.AsSeconds(StartFrameOffset);
 	}
 
 	return AnimPosition;
-
-}
\ No newline at end of file
+}
diff --git a/Engine/Plugins/Experimental/GeometryCache/Source/GeometryCacheTracks/Private/MovieSceneGeometryCacheTemplate.h b/Engine/Plugins/Experimental/GeometryCache/Source/GeometryCacheTracks/Private/MovieSceneGeometryCacheTemplate.h
--- a/Engine/Plugins/Experimental/GeometryCache/Source/GeometryCacheTracks/Private/MovieSceneGeometryCacheTemplate.h
+++ b/Engine/Plugins/Experimental/GeometryCache/Source/GeometryCacheTracks/Private/MovieSceneGeometryCacheTemplate.h
@@ -9,20 +9,18 @@
 #include "MovieSceneGeometryCacheTemplate.generated.h"
 
 USTRUCT()
-struct FMovieSceneGeometryCacheSectionTemplateParameters
+struct FMovieSceneGeometryCacheSectionTemplateParameters : public FMovieSceneGeometryCacheParams
 {
 	GENERATED_BODY()
 
-	FMovieSceneGeometryCacheSectionTemplateParameters() :SectionParams(nullptr) {}
-	FMovieSceneGeometryCacheSectionTemplateParameters(FMovieSceneGeometryCacheParams& BaseParams, FFrameNumber InSectionStartTime, FFrameNumber InSectionEndTime)
-		: SectionParams(&BaseParams)
+		FMovieSceneGeometryCacheSectionTemplateParameters() {}
+	FMovieSceneGeometryCacheSectionTemplateParameters(const FMovieSceneGeometryCacheParams& BaseParams, FFrameNumber InSectionStartTime, FFrameNumber InSectionEndTime)
+		: FMovieSceneGeometryCacheParams(BaseParams)
 		, SectionStartTime(InSectionStartTime)
 		, SectionEndTime(InSectionEndTime)
 	{}
 
-	float MapTimeToAnimation(FFrameTime InPosition, FFrameRate InFrameRate) const;
-
-	FMovieSceneGeometryCacheParams* SectionParams;
+	float MapTimeToAnimation(float ComponentDuration, FFrameTime InPosition, FFrameRate InFrameRate) const;
 
 	UPROPERTY()
 	FFrameNumber SectionStartTime;
diff --git a/Engine/Plugins/Experimental/GeometryCache/Source/GeometryCacheTracks/Private/MovieSceneGeometryCacheTrack.cpp b/Engine/Plugins/Experimental/GeometryCache/Source/GeometryCacheTracks/Private/MovieSceneGeometryCacheTrack.cpp
--- a/Engine/Plugins/Experimental/GeometryCache/Source/GeometryCacheTracks/Private/MovieSceneGeometryCacheTrack.cpp
+++ b/Engine/Plugins/Experimental/GeometryCache/Source/GeometryCacheTracks/Private/MovieSceneGeometryCacheTrack.cpp
@@ -38,7 +38,8 @@ UMovieSceneSection* UMovieSceneGeometryCacheTrack::AddNewAnimation(FFrameNumber
 	{
 		FFrameTime AnimationLength = GeomCacheComp->GetDuration()* GetTypedOuter<UMovieScene>()->GetTickResolution();
 		NewSection->InitialPlacementOnRow(AnimationSections, KeyTime, AnimationLength.FrameNumber.Value, INDEX_NONE);
-		NewSection->Params.GeometryCache = GeomCacheComp;
+		
+		NewSection->Params.GeometryCacheAsset = (GeomCacheComp->GetGeometryCache());
 	}
 
 	AddSection(*NewSection);
