Fix for build warning
#trivial

diff --git a/Engine/Plugins/Experimental/GeometryCache/Source/GeometryCache/Private/GeometryCache.cpp b/Engine/Plugins/Experimental/GeometryCache/Source/GeometryCache/Private/GeometryCache.cpp
--- a/Engine/Plugins/Experimental/GeometryCache/Source/GeometryCache/Private/GeometryCache.cpp
+++ b/Engine/Plugins/Experimental/GeometryCache/Source/GeometryCache/Private/GeometryCache.cpp
@@ -175,7 +175,6 @@ int32 UGeometryCache::GetFrameAtTime(const float Time) const
 	const int32 NumberOfFrames = GetEndFrame() - GetStartFrame() + 1;;
 	const float FrameTime = NumberOfFrames > 1 ? Duration / (float)(NumberOfFrames - 1) : 0.0f;
 	const int32 NormalizedFrame = FMath::Clamp(FMath::RoundToInt(Time / FrameTime), 0, NumberOfFrames - 1);
-	const int32 StartFrame =  GetStartFrame();
 	return StartFrame + NormalizedFrame; 
 
 }
