Moving over Geometry Cache fixes and improvements:

"Trying to address issue with Geometry Cache streaming while scrubbing in Sequencer Time line.
#info Seems like we can get two chunk IO requests with multiple instances of the same Cache, either in level or due to thumbnail rendering.
#fix Remove hard Check and make conditional, in case of duplicate requests see if one has already been fulfilled and otherwise process the first one to finish
#misc Rename Cunk to Chunk

Geometry Cache motion vectors being incorrect
#fix Geometry Cache can be rendered without being updated, this means that no time has elapsed and the rendered positions haven't changed. So now only calculate previous position if an update has happened for this frame, otherwise use the current position as previous since we are in a frozen state.

- Make sure we reset elapsed time when resetting from Sequencer side of things
- Mark ElapsedTime as transient, it was being saved which could cause streaming hiccups on frame 0

diff --git a/Engine/Plugins/Experimental/GeometryCache/Source/GeometryCache/Classes/GeometryCacheCodecV1.h b/Engine/Plugins/Experimental/GeometryCache/Source/GeometryCache/Classes/GeometryCacheCodecV1.h
--- a/Engine/Plugins/Experimental/GeometryCache/Source/GeometryCache/Classes/GeometryCacheCodecV1.h
+++ b/Engine/Plugins/Experimental/GeometryCache/Source/GeometryCache/Classes/GeometryCacheCodecV1.h
@@ -18,7 +18,6 @@ struct FGeometryCacheCodecRenderStateV1 : FGeometryCacheCodecRenderStateBase
 {
 	FGeometryCacheCodecRenderStateV1(const TArray<int32> &SetTopologyRanges); 
 	~FGeometryCacheCodecRenderStateV1();	
-	ICodecDecoder* Decoder;
 	
 	virtual bool DecodeSingleFrame(FGeometryCacheCodecDecodeArguments &Args) override;
 };
diff --git a/Engine/Plugins/Experimental/GeometryCache/Source/GeometryCache/Classes/GeometryCacheComponent.h b/Engine/Plugins/Experimental/GeometryCache/Source/GeometryCache/Classes/GeometryCacheComponent.h
--- a/Engine/Plugins/Experimental/GeometryCache/Source/GeometryCache/Classes/GeometryCacheComponent.h
+++ b/Engine/Plugins/Experimental/GeometryCache/Source/GeometryCache/Classes/GeometryCacheComponent.h
@@ -179,6 +179,8 @@ public:
 	void SetManualTick(bool bInManualTick);
 	bool GetManualTick() const;
 
+	void ResetAnimationTime();
+
 	UFUNCTION(BlueprintCallable, Category = "Components|GeometryCache")
 	void TickAtThisTime(const float Time, bool bInIsRunning, bool bInBackwards, bool bInIsLooping);
 
@@ -234,7 +236,7 @@ protected:
 	UPROPERTY(VisibleAnywhere, Category = GeometryCache)
 	int32 NumTracks;
 
-	UPROPERTY(VisibleAnywhere, Category = GeometryCache)
+	UPROPERTY(VisibleAnywhere, transient, Category = GeometryCache)
 	float ElapsedTime;
 
 	/** Component local space bounds of geometry cache mesh */
diff --git a/Engine/Plugins/Experimental/GeometryCache/Source/GeometryCache/Private/CodecV1.cpp b/Engine/Plugins/Experimental/GeometryCache/Source/GeometryCache/Private/CodecV1.cpp
--- a/Engine/Plugins/Experimental/GeometryCache/Source/GeometryCache/Private/CodecV1.cpp
+++ b/Engine/Plugins/Experimental/GeometryCache/Source/GeometryCache/Private/CodecV1.cpp
@@ -832,18 +832,6 @@ void FCodecV1Encoder::WriteTables()
 
 #endif // WITH_EDITOR
 
-FCodecV1Decoder::FCodecV1Decoder()
-{
-	// Precalculate table mapping symbol index to non-raw bits. ((Sign ? -2 : 1) << NumRawBits)
-	for (int32 NumRawBits = 1; NumRawBits <= 30; NumRawBits++)
-	{
-		for (int32 Sign = 0; Sign <= 1; Sign++)
-		{
-			HighBitsLUT[2 + Sign + NumRawBits * 2] = (Sign ? -2 : 1) << NumRawBits;
-		}
-	}	
-}
-
 void FCodecV1Decoder::DecodeIndexStream(FHuffmanBitStreamReader& Reader, uint32* Stream, uint64 ElementOffset, uint32 ElementCount)
 {
 	const uint8* RawElementData = (const uint8*)Stream;
@@ -1045,6 +1033,19 @@ void FCodecV1Decoder::ReadCodedStreamDescription(FHuffmanBitStreamReader& Reader
 	VertexInfo.bConstantIndices = (ReadBits(Reader, 1) == 1);
 }
 
+int32 FCodecV1Decoder::CachedHighBitsLUT[64];
+void FCodecV1Decoder::InitLUT()
+{
+	// Precalculate table mapping symbol index to non-raw bits. ((Sign ? -2 : 1) << NumRawBits)
+	for (int32 NumRawBits = 1; NumRawBits <= 30; NumRawBits++)
+	{
+		for (int32 Sign = 0; Sign <= 1; Sign++)
+		{
+			CachedHighBitsLUT[2 + Sign + NumRawBits * 2] = (Sign ? -2 : 1) << NumRawBits;
+		}
+	}
+}
+
 DECLARE_CYCLE_STAT(TEXT("FCodecV1Decoder"), STAT_CodecV1Decoder, STATGROUP_GeometryCache);
 DECLARE_CYCLE_STAT(TEXT("SetupAndReadTables"), STAT_SetupAndReadTables, STATGROUP_GeometryCache);
 DECLARE_CYCLE_STAT(TEXT("DecodeIndexStream"), STAT_DecodeIndexStream, STATGROUP_GeometryCache);
@@ -1278,6 +1279,6 @@ int32 FCodecV1Decoder::ReadInt32(FHuffmanBitStreamReader& Reader, FHuffmanDecode
 	{
 		// At least one raw bit.
 		int32 NumRawBits = (Packed - 2) >> 1;
-		return ReadBitsNoRefill(Reader, NumRawBits) + HighBitsLUT[Packed];
+		return ReadBitsNoRefill(Reader, NumRawBits) + CachedHighBitsLUT[Packed];
 	}
 }
\ No newline at end of file
diff --git a/Engine/Plugins/Experimental/GeometryCache/Source/GeometryCache/Private/CodecV1.h b/Engine/Plugins/Experimental/GeometryCache/Source/GeometryCache/Private/CodecV1.h
--- a/Engine/Plugins/Experimental/GeometryCache/Source/GeometryCache/Private/CodecV1.h
+++ b/Engine/Plugins/Experimental/GeometryCache/Source/GeometryCache/Private/CodecV1.h
@@ -287,7 +287,9 @@ private:
 class FCodecV1Decoder : public ICodecDecoder
 {
 public:
-	FCodecV1Decoder();
+	FCodecV1Decoder() {}
+
+	static void InitLUT();
 
 	/**
 	* Read a frame's bit stream and decode the frame
@@ -344,6 +346,6 @@ private:
 	/** Any context information to decode a frame in a sequence of frames, such as any Huffman tables used. */
 	FCodecV1DecodingContext DecodingContext;	
 
-	int32 HighBitsLUT[64];
+	static int32 CachedHighBitsLUT[64];
 };
 
diff --git a/Engine/Plugins/Experimental/GeometryCache/Source/GeometryCache/Private/GeometryCacheCodecV1.cpp b/Engine/Plugins/Experimental/GeometryCache/Source/GeometryCache/Private/GeometryCacheCodecV1.cpp
--- a/Engine/Plugins/Experimental/GeometryCache/Source/GeometryCache/Private/GeometryCacheCodecV1.cpp
+++ b/Engine/Plugins/Experimental/GeometryCache/Source/GeometryCache/Private/GeometryCacheCodecV1.cpp
@@ -123,13 +123,10 @@ void UGeometryCacheCodecV1::CodeFrame(const FGeometryCacheCodecEncodeArguments &
 
 FGeometryCacheCodecRenderStateV1::FGeometryCacheCodecRenderStateV1(const TArray<int32> &SetTopologyRanges) : FGeometryCacheCodecRenderStateBase(SetTopologyRanges)
 {
-	Decoder = new FCodecV1Decoder();
 }
 
 FGeometryCacheCodecRenderStateV1::~FGeometryCacheCodecRenderStateV1()
 {	
-	delete Decoder;
-	Decoder = nullptr;
 }
 
 DECLARE_CYCLE_STAT(TEXT("Deserialize Mesh"), STAT_DeserializeMeshV1, STATGROUP_GeometryCache);
@@ -153,7 +150,8 @@ bool FGeometryCacheCodecRenderStateV1::DecodeSingleFrame(FGeometryCacheCodecDeco
 
 	{
 		SCOPE_CYCLE_COUNTER(STAT_DeserializeMeshV1);
-		Decoder->DecodeFrameData(Ar, Args.OutMeshData);
+		FCodecV1Decoder Decoder;
+		Decoder.DecodeFrameData(Ar, Args.OutMeshData);
 	}
 	IGeometryCacheStreamingManager::Get().UnmapChunk(Args.Track.GetTrack(), Args.FrameIdentifier);
 	return true;
diff --git a/Engine/Plugins/Experimental/GeometryCache/Source/GeometryCache/Private/GeometryCacheComponent.cpp b/Engine/Plugins/Experimental/GeometryCache/Source/GeometryCache/Private/GeometryCacheComponent.cpp
--- a/Engine/Plugins/Experimental/GeometryCache/Source/GeometryCache/Private/GeometryCacheComponent.cpp
+++ b/Engine/Plugins/Experimental/GeometryCache/Source/GeometryCache/Private/GeometryCacheComponent.cpp
@@ -148,6 +148,11 @@ bool UGeometryCacheComponent::GetManualTick() const
 	return bManualTick;
 }
 
+void UGeometryCacheComponent::ResetAnimationTime()
+{
+	ElapsedTime = 0.0f;
+}
+
 void UGeometryCacheComponent::TickAtThisTime(const float Time, bool bInIsRunning, bool bInBackwards, bool bInIsLooping)
 {
 	if (bManualTick && GeometryCache && bRunning)
diff --git a/Engine/Plugins/Experimental/GeometryCache/Source/GeometryCache/Private/GeometryCacheModule.cpp b/Engine/Plugins/Experimental/GeometryCache/Source/GeometryCache/Private/GeometryCacheModule.cpp
--- a/Engine/Plugins/Experimental/GeometryCache/Source/GeometryCache/Private/GeometryCacheModule.cpp
+++ b/Engine/Plugins/Experimental/GeometryCache/Source/GeometryCache/Private/GeometryCacheModule.cpp
@@ -4,6 +4,7 @@
 #if WITH_EDITOR
 #include "GeometryCacheEdModule.h"
 #endif // WITH_EDITOR
+#include "CodecV1.h"
 
 IMPLEMENT_MODULE(FGeometryCacheModule, GeometryCache)
 
@@ -12,6 +13,8 @@ void FGeometryCacheModule::StartupModule()
 #if WITH_EDITOR
 	FGeometryCacheEdModule& Module = FModuleManager::LoadModuleChecked<FGeometryCacheEdModule>(TEXT("GeometryCacheEd"));
 #endif
+
+	FCodecV1Decoder::InitLUT();
 }
 
 void FGeometryCacheModule::ShutdownModule()
diff --git a/Engine/Plugins/Experimental/GeometryCache/Source/GeometryCache/Private/GeometryCacheSceneProxy.cpp b/Engine/Plugins/Experimental/GeometryCache/Source/GeometryCache/Private/GeometryCacheSceneProxy.cpp
--- a/Engine/Plugins/Experimental/GeometryCache/Source/GeometryCache/Private/GeometryCacheSceneProxy.cpp
+++ b/Engine/Plugins/Experimental/GeometryCache/Source/GeometryCache/Private/GeometryCacheSceneProxy.cpp
@@ -54,6 +54,7 @@ FGeometryCacheSceneProxy::FGeometryCacheSceneProxy(UGeometryCacheComponent* Comp
 	bLooping = Component->IsLooping();
 	bAlwaysHasVelocity = true;
 	PlaybackSpeed = (Component->IsPlaying()) ? Component->GetPlaybackSpeed() : 0.0f;
+	UpdatedFrameNum = 0;
 
 	// Copy each section
 	const int32 NumTracks = Component->TrackSections.Num();
@@ -324,17 +325,18 @@ void FGeometryCacheSceneProxy::CreateMeshBatch(
 				&& (TrackProxy->NextFrameMeshData->Positions.Num() == TrackProxy->NextFrameMeshData->MotionVectors.Num());
 
 				if (!bHasMotionVectors)
-				{
-					UserData.MotionBlurDataExtension = FVector::OneVector;
-					UserData.MotionBlurDataOrigin = FVector::ZeroVector;
-					UserData.MotionBlurPositionScale = 0.0f;
-				}
-				else
-				{
-					UserData.MotionBlurDataExtension = FVector::OneVector * PlaybackSpeed;
+	            {
+		            const float PreviousPositionScale = (GFrameNumber <= UpdatedFrameNum) ? 1.f : 0.f;
+					UserData.MotionBlurDataExtension = FVector::OneVector * PreviousPositionScale;
 					UserData.MotionBlurDataOrigin = FVector::ZeroVector;
-					UserData.MotionBlurPositionScale = 1.0f;
-				}
+					UserData.MotionBlurPositionScale = 1.f - PreviousPositionScale;
+	            }
+	            else
+	            {
+		            UserData.MotionBlurDataExtension = FVector::OneVector * PlaybackSpeed;
+		            UserData.MotionBlurDataOrigin = FVector::ZeroVector;
+		            UserData.MotionBlurPositionScale = 1.0f;
+	            }
 
 				if (IsRayTracingEnabled())
 				{
@@ -539,6 +541,7 @@ void FGeometryCacheSceneProxy::UpdateAnimation(float NewTime, bool bNewLooping,
 	bLooping = bNewLooping;
 	bIsPlayingBackwards = bNewIsPlayingBackwards;
 	PlaybackSpeed = NewPlaybackSpeed;
+	UpdatedFrameNum = GFrameNumber + 1;
 
 	if (IsRayTracingEnabled())
 	{
@@ -1017,7 +1020,6 @@ void FGeomCacheVertexFactory::Init(const FVertexBuffer* PositionBuffer, const FV
 			{
 			Init_RenderThread(PositionBuffer, MotionBlurDataBuffer, TangentXBuffer, TangentZBuffer, TextureCoordinateBuffer, ColorBuffer);
 		});
-		FlushRenderingCommands();
 	}
 }
 
diff --git a/Engine/Plugins/Experimental/GeometryCache/Source/GeometryCache/Private/GeometryCacheSceneProxy.h b/Engine/Plugins/Experimental/GeometryCache/Source/GeometryCache/Private/GeometryCacheSceneProxy.h
--- a/Engine/Plugins/Experimental/GeometryCache/Source/GeometryCache/Private/GeometryCacheSceneProxy.h
+++ b/Engine/Plugins/Experimental/GeometryCache/Source/GeometryCache/Private/GeometryCacheSceneProxy.h
@@ -217,6 +217,7 @@ private:
 	bool bLooping;
 	bool bIsPlayingBackwards;
 	float PlaybackSpeed;
+	uint32 UpdatedFrameNum;
 
 	/** Array of Track Proxies */
 	TArray<FGeomCacheTrackProxy*> Tracks;
diff --git a/Engine/Plugins/Experimental/GeometryCache/Source/GeometryCache/Private/StreamingGeometryCacheData.cpp b/Engine/Plugins/Experimental/GeometryCache/Source/GeometryCache/Private/StreamingGeometryCacheData.cpp
--- a/Engine/Plugins/Experimental/GeometryCache/Source/GeometryCache/Private/StreamingGeometryCacheData.cpp
+++ b/Engine/Plugins/Experimental/GeometryCache/Source/GeometryCache/Private/StreamingGeometryCacheData.cpp
@@ -329,42 +329,45 @@ void FStreamingGeometryCacheData::ProcessCompletedChunks()
 	//Note: This function should only be called from code which owns the CriticalSection
 	check(IsInGameThread() || IsInRenderingThread());
 
-	FCompletedChunk CompletedCunk;
-	while (CompletedChunks.Dequeue(CompletedCunk))
+	FCompletedChunk CompletedChunk;
+	while (CompletedChunks.Dequeue(CompletedChunk))
 	{
-		FResidentChunk* Chunk = Chunks.Find(CompletedCunk.LoadedChunkIndex);
+		FResidentChunk* Chunk = Chunks.Find(CompletedChunk.LoadedChunkIndex);
 		if (!Chunk)
 		{
 			UE_LOG(LogGeoCaStreaming, Error, TEXT("Got a stray async read request"));
 			return;
 		}
 
-		check(CompletedCunk.ReadRequest == Chunk->IORequest);
-
-		// Check to see if we successfully managed to load anything
-		uint8* Mem = CompletedCunk.ReadRequest->GetReadResults();
-		if (Mem)
+		// Chunks can be queued up multiple times when scrubbing, either the request has already been fulfilled (nullptr) or it could be waiting for a different IO request
+		if (CompletedChunk.ReadRequest == Chunk->IORequest || Chunk->IORequest != nullptr)
 		{
-			Chunk->Memory = Mem;
-			ChunksAvailable.Add(CompletedCunk.LoadedChunkIndex);
-			ChunksRequested.Remove(CompletedCunk.LoadedChunkIndex);
-			DEC_DWORD_STAT(STAT_OutstandingRequests);
-			INC_MEMORY_STAT_BY(STAT_ChunkDataResident, Chunk->DataSize);
-			INC_MEMORY_STAT_BY(STAT_ChunkDataStreamed, Chunk->DataSize);
-			IGeometryCacheStreamingManager::Get().IoBandwidth.Add(Chunk->DataSize);
-		}
-		else
-		{
-			UE_LOG(LogGeoCaStreaming, Error, TEXT("Async loading request failed!"));
-			ChunksRequested.Remove(CompletedCunk.LoadedChunkIndex);
-			// Fix me do we want to recover from this? Granite simply reschedules requests
-			// as it may have failed for transient reasons (buffer contention, ...)
+			// Check to see if we successfully managed to load anything
+			uint8* Mem = CompletedChunk.ReadRequest->GetReadResults();
+			if (Mem)
+			{
+				Chunk->Memory = Mem;
+				ChunksAvailable.Add(CompletedChunk.LoadedChunkIndex);
+				ChunksRequested.Remove(CompletedChunk.LoadedChunkIndex);
+				DEC_DWORD_STAT(STAT_OutstandingRequests);
+				INC_MEMORY_STAT_BY(STAT_ChunkDataResident, Chunk->DataSize);
+				INC_MEMORY_STAT_BY(STAT_ChunkDataStreamed, Chunk->DataSize);
+				IGeometryCacheStreamingManager::Get().IoBandwidth.Add(Chunk->DataSize);
+			}
+			else
+			{
+				UE_LOG(LogGeoCaStreaming, Error, TEXT("Async loading request failed!"));
+				ChunksRequested.Remove(CompletedChunk.LoadedChunkIndex);
+				// Fix me do we want to recover from this? Granite simply reschedules requests
+				// as it may have failed for transient reasons (buffer contention, ...)
+			}
+
+			Chunk->IORequest = nullptr;
 		}
 
 		// Clean up the now fully processed IO request
-		check(Chunk->IORequest->PollCompletion());
-		delete Chunk->IORequest;
-		Chunk->IORequest = nullptr;
+		check(CompletedChunk.ReadRequest->PollCompletion());
+		delete CompletedChunk.ReadRequest;
 	}
 }
 
@@ -403,7 +406,6 @@ const uint8* FStreamingGeometryCacheData::MapChunk(uint32 ChunkIndex, uint32* Ou
 			*OutChunkSize = ResidentChunk->DataSize;
 		}
 		ResidentChunk->Refcount++;
-		check(ResidentChunk->Refcount == 1);
 		return (uint8*)ResidentChunk->Memory;
 	}
 }
@@ -419,7 +421,6 @@ void FStreamingGeometryCacheData::UnmapChunk(uint32 ChunkIndex)
 		checkf(ResidentChunk->Refcount > 0, TEXT("Map/Unmap out of balance. Make sure you unmap once fore every map."));
 		checkf(ChunksAvailable.Contains(ChunkIndex) || ChunksEvicted.Contains(ChunkIndex), TEXT("Tried to unmap a chunk in an invalid state."));
 		ResidentChunk->Refcount--;
-		check(ResidentChunk->Refcount == 0);
 	}
 	else
 	{
diff --git a/Engine/Plugins/Experimental/GeometryCache/Source/GeometryCacheTracks/Private/MovieSceneGeometryCacheTemplate.cpp b/Engine/Plugins/Experimental/GeometryCache/Source/GeometryCacheTracks/Private/MovieSceneGeometryCacheTemplate.cpp
--- a/Engine/Plugins/Experimental/GeometryCache/Source/GeometryCacheTracks/Private/MovieSceneGeometryCacheTemplate.cpp
+++ b/Engine/Plugins/Experimental/GeometryCache/Source/GeometryCacheTracks/Private/MovieSceneGeometryCacheTemplate.cpp
@@ -29,6 +29,7 @@ struct FPreAnimatedGeometryCacheTokenProducer : IMovieScenePreAnimatedTokenProdu
 			{
 				UGeometryCacheComponent* Component = CastChecked<UGeometryCacheComponent>(&ObjectToRestore);
 				Component->SetManualTick(bInManualTick);
+				Component->ResetAnimationTime();
 			}
 			bool bInManualTick;
 		};
