Fix geometry cache playback by making sure last animation frame is called by rounding up the frame similarly to what we did for skelelal animationt tracks.

diff --git a/Engine/Plugins/Experimental/GeometryCache/Source/GeometryCacheTracks/Private/MovieSceneGeometryCacheSection.cpp b/Engine/Plugins/Experimental/GeometryCache/Source/GeometryCacheTracks/Private/MovieSceneGeometryCacheSection.cpp
--- a/Engine/Plugins/Experimental/GeometryCache/Source/GeometryCacheTracks/Private/MovieSceneGeometryCacheSection.cpp
+++ b/Engine/Plugins/Experimental/GeometryCache/Source/GeometryCacheTracks/Private/MovieSceneGeometryCacheSection.cpp
@@ -89,8 +89,8 @@ FFrameNumber GetStartOffsetAtTrimTime(FQualifiedFrameTime TrimTime, const FMovie
 TOptional<TRange<FFrameNumber> > UMovieSceneGeometryCacheSection::GetAutoSizeRange() const
 {
 	FFrameRate FrameRate = GetTypedOuter<UMovieScene>()->GetTickResolution();
-
 	FFrameTime AnimationLength = Params.GetSequenceLength() * FrameRate;
+	int32 IFrameNumber = AnimationLength.FrameNumber.Value + (int)(AnimationLength.GetSubFrame() + 0.5f);
 
 	return TRange<FFrameNumber>(GetInclusiveStartFrame(), GetInclusiveStartFrame() + AnimationLength.FrameNumber);
 }
diff --git a/Engine/Plugins/Experimental/GeometryCache/Source/GeometryCacheTracks/Private/MovieSceneGeometryCacheTemplate.cpp b/Engine/Plugins/Experimental/GeometryCache/Source/GeometryCacheTracks/Private/MovieSceneGeometryCacheTemplate.cpp
--- a/Engine/Plugins/Experimental/GeometryCache/Source/GeometryCacheTracks/Private/MovieSceneGeometryCacheTemplate.cpp
+++ b/Engine/Plugins/Experimental/GeometryCache/Source/GeometryCacheTracks/Private/MovieSceneGeometryCacheTemplate.cpp
@@ -87,6 +87,11 @@ void FMovieSceneGeometryCacheSectionTemplate::Evaluate(const FMovieSceneEvaluati
 
 float FMovieSceneGeometryCacheSectionTemplateParameters::MapTimeToAnimation(FFrameTime InPosition, FFrameRate InFrameRate) const
 {
+	FFrameTime AnimationLength = GetSequenceLength() * InFrameRate;
+	int32 LengthInFrames = AnimationLength.FrameNumber.Value + (int)(AnimationLength.GetSubFrame() + 0.5f) + 1;
+	//we only play end if we are not looping, and assuming we are looping if Length is greater than default length;
+	bool bLooping = (SectionEndTime.Value - SectionStartTime.Value) > LengthInFrames;
+
 	InPosition = FMath::Clamp(InPosition, FFrameTime(SectionStartTime), FFrameTime(SectionEndTime-1));
 
 	const float SectionPlayRate = PlayRate;
@@ -95,7 +100,7 @@ float FMovieSceneGeometryCacheSectionTemplateParameters::MapTimeToAnimation(FFra
 	const float SeqLength = GetSequenceLength() - InFrameRate.AsSeconds(StartFrameOffset + EndFrameOffset);
 
 	float AnimPosition = FFrameTime::FromDecimal((InPosition - SectionStartTime).AsDecimal() * AnimPlayRate) / InFrameRate;
-	if (SeqLength > 0.f)
+	if (SeqLength > 0.f && (bLooping || !FMath::IsNearlyEqual(AnimPosition, SeqLength, 1e-4f)))
 	{
 		AnimPosition = FMath::Fmod(AnimPosition, SeqLength);
 	}
@@ -106,4 +111,5 @@ float FMovieSceneGeometryCacheSectionTemplateParameters::MapTimeToAnimation(FFra
 	}
 
 	return AnimPosition;
+
 }
diff --git a/Engine/Plugins/Experimental/GeometryCache/Source/GeometryCacheTracks/Private/MovieSceneGeometryCacheTrack.cpp b/Engine/Plugins/Experimental/GeometryCache/Source/GeometryCacheTracks/Private/MovieSceneGeometryCacheTrack.cpp
--- a/Engine/Plugins/Experimental/GeometryCache/Source/GeometryCacheTracks/Private/MovieSceneGeometryCacheTrack.cpp
+++ b/Engine/Plugins/Experimental/GeometryCache/Source/GeometryCacheTracks/Private/MovieSceneGeometryCacheTrack.cpp
@@ -37,7 +37,8 @@ UMovieSceneSection* UMovieSceneGeometryCacheTrack::AddNewAnimation(FFrameNumber
 	UMovieSceneGeometryCacheSection* NewSection = Cast<UMovieSceneGeometryCacheSection>(CreateNewSection());
 	{
 		FFrameTime AnimationLength = GeomCacheComp->GetDuration()* GetTypedOuter<UMovieScene>()->GetTickResolution();
-		NewSection->InitialPlacementOnRow(AnimationSections, KeyTime, AnimationLength.FrameNumber.Value, INDEX_NONE);
+		int32 IFrameNumber = AnimationLength.FrameNumber.Value + (int)(AnimationLength.GetSubFrame() + 0.5f) + 1;
+		NewSection->InitialPlacementOnRow(AnimationSections, KeyTime, IFrameNumber, INDEX_NONE);
 		NewSection->Params.GeometryCache = GeomCacheComp;
 	}
 
