Merged from CL4841935
Optimized compression codec for GeometryCache.
Approximately an order of magnitude faster decompression at a similar compression ratio.
Not backwards compatible: Requires alembic caches to be reimported!

Changes:
Split HuffmanTable into HuffmanEncodeTable and HuffmanDecodeTable as there is little overlap between the two.
Implemented optimal length-limited huffman code generation using coin-collector + package-merge
Got rid of long huffman as they add extra branches in the decompressor while yielding virtually no better compression than length-limited codes.
Huffman tree is now transmitted as code lengths instead of symbol indices. Simpler and ~0.5mb small on DressA.
Reversed bitstream to little-endian order for faster branchless bit refills.
Keep pending bitstream bits in 64bit staging buffer to avoid repeatedly fetching the same bytes in Peek and Read.
Added variants of Peek and Reed that don't automatically refill bits.
BitStream no longer owns it data and will no longer copy the entire input into its own local buffer.
Got rid of FRingBuffer in the decoding of indices, motion vectors, UVs, normals, colors, etc. They were unnecessary as only the last element was ever read.
Got rid of per-vertex prediction modes. Always use mode 0. The better prediction quality was almost completely negated by the need to send per vertex information about the predictor used.
Got rid of complicated triangle rotation code in Position decoding as it was only necessary for parallelogram prediction.
New Int32 encoding that doesn't branch based on sign and handles the 4 center values without reading raw bits.
Encode Normals and Tangents (8bit values) using direct huffman codes instead of using Int32 (magnitude + raw bits) encoding.

diff --git a/Engine/Plugins/Experimental/GeometryCache/Source/GeometryCache/Private/CodecV1.cpp b/Engine/Plugins/Experimental/GeometryCache/Source/GeometryCache/Private/CodecV1.cpp
--- a/Engine/Plugins/Experimental/GeometryCache/Source/GeometryCache/Private/CodecV1.cpp
+++ b/Engine/Plugins/Experimental/GeometryCache/Source/GeometryCache/Private/CodecV1.cpp
@@ -188,98 +188,6 @@ private:
 	FQuantizer Quantizer;	
 };
 
-int32 FCodecV1SharedTools::EstimateCodingCost(const FIntVector& Value)
-{
-	// Just use the absolute sizes of the components as a representative for a relative coding cost of this vector
-	return FMath::Abs(Value.X) + FMath::Abs(Value.Y) + FMath::Abs(Value.Z);
-}
-
-FIntVector FCodecV1SharedTools::PredictVertex(uint32 CornerIndex, uint32 PreviousTriangleRotationOffsets[3], FRingBuffer<FIntVector, VertexStreamCodingVertexHistorySize>& ReconstructedVertexHistory, uint32 Mode)
-{
-	FIntVector Prediction;
-	
-	switch (Mode)
-	{
-	case 0:
-		{
-			// Mode 0: delta coding, last visited vertex
-			Prediction = ReconstructedVertexHistory[0];
-		}
-		break;
-	case 1:
-		{
-			// Mode 1: delta coding, second to last visited vertex
-			Prediction = ReconstructedVertexHistory[1]; 		
-		}
-		break;
-	case 2:
-		{
-			// Mode 2: delta coding, third to last visited vertex
-			Prediction = ReconstructedVertexHistory[2];			
-		}
-		break;
-	case 3:
-		{
-			// Mode 3: Parallelogram prediction: P = V + U - W				
-			// Calculate offsets of the previous triangle in our vertex history buffer and compensate for possible rotation of the triangle		
-			const uint32 OffsetV = CornerIndex + 2 - PreviousTriangleRotationOffsets[0]; // Shared edge vertex 1
-			const uint32 OffsetU = CornerIndex + 2 - PreviousTriangleRotationOffsets[1]; // Shared edge vertex 2
-			const uint32 OffsetW = CornerIndex + 2 - PreviousTriangleRotationOffsets[2]; // non-shared vertex
-			
-			const FIntVector PredV = ReconstructedVertexHistory[OffsetV];
-			const FIntVector PredU = ReconstructedVertexHistory[OffsetU];
-			const FIntVector PredW = ReconstructedVertexHistory[OffsetW];
-			Prediction = PredV + PredU - PredW; 			
-		}		
-		break;
-	default:
-		{
-			checkf(false, TEXT("Unknown prediction mode, corrupt bitstream or encoder and decoder are out of sync."));
-		}
-	}
-	
-	return Prediction;
-}
-
-bool FCodecV1SharedTools::FindTrianglesSharedEdge(const uint32 Triangle1[3], const uint32 Triangle2[3], uint32 Triangle1RotationOffsets[3])
-{
-	// e.g. Triangle1 = {3, 4, 5}, Triangle2 = {5, 6, 3}
-	//      Returns: true, orderedTriangle1 = {2, 0, 1}
-	//      The shared edge is (orderedTriangle[0], orderedTriangle1[1]), the non-edge vertex is orderedTriangle1[2]
-
-	const bool bIsShared[3] =
-	{
-		Triangle1[0] == Triangle2[0] || Triangle1[0] == Triangle2[1] || Triangle1[0] == Triangle2[2],
-		Triangle1[1] == Triangle2[0] || Triangle1[1] == Triangle2[1] || Triangle1[1] == Triangle2[2],
-		Triangle1[2] == Triangle2[0] || Triangle1[2] == Triangle2[1] || Triangle1[2] == Triangle2[2]
-	};
-	
-	if (bIsShared[0] && bIsShared[1])
-	{
-		Triangle1RotationOffsets[0] = 0;
-		Triangle1RotationOffsets[1] = 1;
-		Triangle1RotationOffsets[2] = 2;
-		return true;
-	}
-
-	if (bIsShared[1] && bIsShared[2])
-	{
-		Triangle1RotationOffsets[0] = 1;
-		Triangle1RotationOffsets[1] = 2;
-		Triangle1RotationOffsets[2] = 0;
-		return true;
-	}
-
-	if (bIsShared[2] && bIsShared[0])
-	{
-		Triangle1RotationOffsets[0] = 2;
-		Triangle1RotationOffsets[1] = 0;
-		Triangle1RotationOffsets[2] = 1;
-		return true;
-	}
-
-	return false;
-}
 // Custom serialization for our const TArray because of the const
 FArchive& operator<<(FArchive &Ar, const TArray<FGeometryCacheMeshBatchInfo>& BatchesInfo)
 {
@@ -349,38 +257,6 @@ void FCodecV1Encoder::EncodeIndexStream(const uint32* Stream, uint64 ElementOffs
 	Stats = FStreamEncodingStatistics(ByteCounter.Read(), ElementCount * sizeof(uint32), Quality);
 }
 
-FIntVector FCodecV1Encoder::FindModeAndPredictVertex(const FIntVector& ValueToCode, uint32 CornerIndex, uint32 PreviousTriangleRotationOffsets[3], FRingBuffer<FIntVector, VertexStreamCodingVertexHistorySize>& ReconstructedVertexHistory, uint32& ChosenMode)
-{
-	// We perform multiple prediction methods and return prediction and method. 
-	// Methods comprise of: delta coding (previous or n-th last value), parallelogram prediction on last seen triangle.
-	// Whether the parallelogram prediction perform a good prediction depends on the mesh optimization the importer/preprocessor did and how much 
-	// coherency there is between successive triangles. No matter the input, the algorithm will never fail. Worst case scenario, prediction is bad 
-	// and we code more bits.	
-	uint32 BestMode = 0;
-	uint32 MinimumCost = UINT32_MAX;
-	FIntVector BestPrediction;
-		
-	for (uint32 Mode = 0; Mode < VertexPredictionModeCount; ++Mode)
-	{
-		// Predict according to this mode
-		FIntVector Prediction = FCodecV1SharedTools::PredictVertex(CornerIndex, PreviousTriangleRotationOffsets, ReconstructedVertexHistory, Mode);
-
-		// Estimate the cost of resulting delta vector
-		FIntVector Delta = ValueToCode - Prediction;		
-		uint32 Cost = FCodecV1SharedTools::EstimateCodingCost(Delta);
-		
-		if (Cost < MinimumCost)
-		{			
-			MinimumCost = Cost;
-			BestMode = Mode;
-			BestPrediction = Prediction;
-		}
-	}
-	
-	ChosenMode = BestMode;
-	return BestPrediction;
-}
-
 void FCodecV1Encoder::EncodePositionStream(const FVector* VertexStream, uint64 VertexElementOffset, uint32 VertexElementCount, const uint32* IndexStream, uint64 IndexElementOffset, uint32 IndexElementCount, FStreamEncodingStatistics& Stats)
 {
 	FBitstreamWriterByteCounter ByteCounter(EncodingContext.Writer);
@@ -403,8 +279,8 @@ void FCodecV1Encoder::EncodePositionStream(const FVector* VertexStream, uint64 V
 	WriteBytes((void*)&Header, sizeof(Header));
 
 	uint32 EncodedVertexCount = 0;
-	FRingBuffer<uint32, VertexStreamCodingIndexHistorySize> IndexHistory(VertexStreamCodingIndexHistorySize); // Previously seen indices
-	FRingBuffer<FIntVector, VertexStreamCodingVertexHistorySize> ReconstructedVertexHistory(VertexStreamCodingVertexHistorySize, FIntVector(0, 0, 0)); // Previously seen positions
+	
+	FIntVector Prediction(0, 0, 0);	// Previously seen position
 	int64 MaxEncounteredIndex = -1;
 		
 	FQualityMetric QualityMetric;
@@ -412,93 +288,44 @@ void FCodecV1Encoder::EncodePositionStream(const FVector* VertexStream, uint64 V
 	const uint8* RawElementDataIndices = (const uint8*)IndexStream;
 	const uint8* RawElementDataVertices = (const uint8*)VertexStream;
 
-	uint32 PreviousTriangleCornerIndices[3] = { 0 };
-	uint32 FaceCount = IndexElementCount / 3;
-
 	// Walk over indices/triangles
-	for (uint32 FaceIdx = 0; FaceIdx < FaceCount; ++FaceIdx, RawElementDataIndices += 3 * IndexElementOffset)
+	for (uint32 IndexIdx = 0; IndexIdx < IndexElementCount; ++IndexIdx)
 	{
-		uint32 CornerIndices[3];
-		CornerIndices[0] = *(uint32*)(RawElementDataIndices);
-		CornerIndices[1] = *(uint32*)(RawElementDataIndices + IndexElementOffset);
-		CornerIndices[2] = *(uint32*)(RawElementDataIndices + 2 * IndexElementOffset);
-
-		// Find if this triangle shares an edge with the previous one, if so, get the rotation offsets to rotate the triangle's indices 
-		// We want a well-defined order of indices so we can look up the right vertices to apply the parallelogram predictor
-		uint32 PreviousTriangleRotationOffsets[3] = { 0, 0, 0 };
-		bool bTrianglesShareEdge = FCodecV1SharedTools::FindTrianglesSharedEdge(PreviousTriangleCornerIndices, CornerIndices, PreviousTriangleRotationOffsets);
+		const uint32 IndexValue = *(const uint32*)RawElementDataIndices;
+		RawElementDataIndices += IndexElementOffset;
 
-		// Walk over triangle corners
-		for (uint32 CornerIndex = 0; CornerIndex < 3; ++CornerIndex)
+		if ((int64)IndexValue > MaxEncounteredIndex)
 		{
-			uint32 IndexValue = CornerIndices[CornerIndex];
-			bool NewlyEncounteredVertex = (int64)IndexValue > MaxEncounteredIndex; // We rely on indices being references in order
-			MaxEncounteredIndex = FMath::Max((int64)IndexValue, MaxEncounteredIndex);
-
-			if (NewlyEncounteredVertex)
-			{
-				// Code a newly encountered vertex
-				FVector VertexValue = *(FVector*)RawElementDataVertices;
-				RawElementDataVertices += VertexElementOffset;
+			MaxEncounteredIndex = (int64)IndexValue;
 
-				// Quantize
-				FIntVector Encoded = Quantizer.Quantize(VertexValue);
+			// Code a newly encountered vertex
+			FVector VertexValue = *(FVector*)RawElementDataVertices;
+			RawElementDataVertices += VertexElementOffset;
 
-				// Translate to center
-				FIntVector EncodedCentered = Encoded - QuantizedTranslationToCenter;
+			// Quantize
+			FIntVector Encoded = Quantizer.Quantize(VertexValue);
 
-				// Find best prediction mode and prediction using our triangle history
-				uint32 PredictionMode;
-				FIntVector Prediction = FindModeAndPredictVertex(EncodedCentered, CornerIndex, PreviousTriangleRotationOffsets, ReconstructedVertexHistory, PredictionMode);
+			// Translate to center
+			FIntVector EncodedCentered = Encoded - QuantizedTranslationToCenter;
 
-				// Residual to code
-				FIntVector Residual = EncodedCentered - Prediction;
+			// Residual to code
+			FIntVector Residual = EncodedCentered - Prediction;
 				
-				// Write mode
-				WriteSymbol(EncodingContext.VertexPredictionModeTable, PredictionMode);
-								
-				// Write residual
-				WriteInt32(EncodingContext.ResidualVertexPosTable, Residual.X);
-				WriteInt32(EncodingContext.ResidualVertexPosTable, Residual.Y);
-				WriteInt32(EncodingContext.ResidualVertexPosTable, Residual.Z);				
-
-				EncodedVertexCount++;
+			// Write residual
+			WriteInt32(EncodingContext.ResidualVertexPosTable, Residual.X);
+			WriteInt32(EncodingContext.ResidualVertexPosTable, Residual.Y);
+			WriteInt32(EncodingContext.ResidualVertexPosTable, Residual.Z);
+
+			EncodedVertexCount++;
 				
-				// Store previous encountered values
-				FIntVector Reconstructed = Prediction + Residual;
-				ReconstructedVertexHistory.Push(Reconstructed);
-
-				// Calculate error
-				FVector DequantReconstructed = Quantizer.Dequantize(Reconstructed);
-				QualityMetric.Register(VertexValue, DequantReconstructed);
-			}
-			else // NewlyEncounteredVertex == false
-			{
-				// Keep our vertex history consistent with the indices we have seen.
-				// We fetch from our previously reconstructed vertex history. Our history is limited in size, so
-				// we do a explicit search and if the value is not found, we use the last seen vector as a safe prediction
-				FIntVector EncounteredVertex = ReconstructedVertexHistory[0];
-
-				for (uint32 Index = 0; Index < IndexHistory.Num(); ++Index)
-				{
-					if (IndexHistory[Index] == IndexValue)
-					{
-						// We have it in our history buffer
-						EncounteredVertex = ReconstructedVertexHistory[Index];
-						break;
-					}
-				}
-
-				ReconstructedVertexHistory.Push(EncounteredVertex);				
-			}
-
-			IndexHistory.Push(IndexValue);
-		}
+			// Store previous encountered values
+			FIntVector Reconstructed = Prediction + Residual;
 
-		// Save our triangle
-		PreviousTriangleCornerIndices[0] = CornerIndices[0];
-		PreviousTriangleCornerIndices[1] = CornerIndices[1];
-		PreviousTriangleCornerIndices[2] = CornerIndices[2];
+			// Calculate error
+			FVector DequantReconstructed = Quantizer.Dequantize(Reconstructed);
+			QualityMetric.Register(VertexValue, DequantReconstructed);
+			Prediction = Reconstructed;
+		}
 	}
 
 	// Gather rate and quality statistics
@@ -547,12 +374,12 @@ void FCodecV1Encoder::EncodeColorStream(const FColor* Stream, uint64 ElementOffs
 	Stats = FStreamEncodingStatistics(ByteCounter.Read(), ElementCount * sizeof(FColor), Quality);
 }
 
-void FCodecV1Encoder::EncodeNormalStream(const FPackedNormal* Stream, uint64 ElementOffsetBytes, uint32 ElementCount, FHuffmanTable& Table, FStreamEncodingStatistics& Stats)
+void FCodecV1Encoder::EncodeNormalStream(const FPackedNormal* Stream, uint64 ElementOffsetBytes, uint32 ElementCount, FHuffmanEncodeTable& Table, FStreamEncodingStatistics& Stats)
 {
 	FBitstreamWriterByteCounter ByteCounter(EncodingContext.Writer);
 
-	FRingBuffer<FIntVector4, NormalStreamCodingHistorySize> ReconstructedHistory(NormalStreamCodingHistorySize, FIntVector4(128, 128, 128, 128));
-
+	uint8 x = 128, y = 128, z = 128, w = 128;
+	
 	const uint8* RawElementData = (const uint8*)Stream;
 
 	// Walk over colors
@@ -560,20 +387,22 @@ void FCodecV1Encoder::EncodeNormalStream(const FPackedNormal* Stream, uint64 Ele
 	{
 		// Load data
 		FPackedNormal& NormalValue = *(FPackedNormal*)RawElementData;		
-		FIntVector4 Value(NormalValue.Vector.X, NormalValue.Vector.Y, NormalValue.Vector.Z, NormalValue.Vector.W);
-
-		FIntVector4 Prediction = ReconstructedHistory[0]; // Delta coding
-		FIntVector4 Residual = FCodecV1SharedTools::SubtractVector4(Value, Prediction); // Residual = Value - Prediction
+		
+		int8 dx = NormalValue.Vector.X - x;
+		int8 dy = NormalValue.Vector.Y - y;
+		int8 dz = NormalValue.Vector.Z - z;
+		int8 dw = NormalValue.Vector.W - w;
 
 		// Write residual	
-		WriteInt32(Table, Residual.X);
-		WriteInt32(Table, Residual.Y);
-		WriteInt32(Table, Residual.Z);
-		WriteInt32(Table, Residual.W);
+		WriteSymbol(Table, (uint8)dx);
+		WriteSymbol(Table, (uint8)dy);
+		WriteSymbol(Table, (uint8)dz);
+		WriteSymbol(Table, (uint8)dw);
 
-		// Decode as the decoder would and keep the result for future prediction
-		FIntVector4 Reconstructed = FCodecV1SharedTools::SumVector4(Prediction, Residual); // Reconstructed = Prediction + Residual
-		ReconstructedHistory.Push(Reconstructed);
+		x = NormalValue.Vector.X;
+		y = NormalValue.Vector.Y;
+		z = NormalValue.Vector.Z;
+		w = NormalValue.Vector.W;
 	}
 
 	// Gather rate and quality statistics
@@ -857,29 +686,46 @@ void FCodecV1Encoder::WriteBytes(const void* Data, int64 NumBytes)
 }
 
 
-void FCodecV1Encoder::WriteInt32(FHuffmanTable& ValueTable, int32 Value)
+void FCodecV1Encoder::WriteInt32(FHuffmanEncodeTable& ValueTable, int32 Value)
 {
-	// We code the length of the value (huffman table), followed by the raw bits. Small values take up very small amount of bits.
-	const int32 ZeroValue32BitInt = 32;	// Zero value used to shift interval [32,31] to [0,63]
-
-	int32 AbsValue = FMath::Abs(Value);
-
-	if (Value == 0)
+	// It is impractical to entropy code an entire integer, so we split it into an entropy coded magnitude followed by a number of raw bits.
+	// The reasoning is that usually most of the redundancy is in the magnitude of the number, not the exact value.
+	
+	// Positive values are encoded as the index k of the first 1-bit (at most 30) followed by the remaining k bits encoded as raw bits.
+	// Negative values are handled symmetrically, but using the index of the first 0-bit.
+	// With one symbol for every bit length and sign, the set of reachable number is 2 * (2^0 + 2^1 + ... + 2^30) = 2 * (2^31 - 1) = 2^32 - 2
+	// To cover all 2^32 possible integer values, we have use separate codes for the remaining two symbols (with no raw bits).
+	// The total number of symbols is 2 * 31 + 2 = 64
+	
+	if (Value >= -2 && Value <= 1)
 	{
-		WriteSymbol(ValueTable, ZeroValue32BitInt);
+		// 4 center values have no raw bits. One more negative values than positive,
+		// so we have an equal number of positive and negative values remaining.
+		WriteSymbol(ValueTable, Value + 2);	// [-2, 1] -> [0, 3]
 	}
-	else if (Value > 0)
+	else
 	{
-		int32 HiBit = HighestSetBit(AbsValue);	// Find most significant bit k, i.e., number of bits
-		WriteSymbol(ValueTable, ZeroValue32BitInt + (HiBit + 1));	// Code number of bits, recenter around new zero, e.g., 32 + 10 bits
-		WriteBits(AbsValue - (1 << HiBit), HiBit);	// Code remaining k bits verbatim
+		// At least one raw bit.
+		if (Value >= 0)
+		{
+			// Value >= 2
+			int32 NumRawBits = HighestSetBit(Value);	// Find first 1-bit. 1 <= NumRawBits <= 30.
+			int32 Packed = 2 + NumRawBits * 2;			// First positive code is 4
+			WriteSymbol(ValueTable, Packed);
+			int32 RawBits = Value - (1 << NumRawBits);
+			WriteBits(RawBits, NumRawBits);
+		}
+		else
+		{
+			// Value <= -3
+			int32 NumRawBits = HighestSetBit(~Value);	// Find first 0-bit. 1 <= NumRawBits <= 30.
+			int32 Packed = 3 + NumRawBits * 2;			// First negative code is 5
+			WriteSymbol(ValueTable, Packed);
+			int32 RawBits = Value & ~(-1 << NumRawBits);
+			WriteBits(RawBits, NumRawBits);
+		}
 	}
-	else /*if (Value < 0)*/
-	{
-		int32 HiBit = HighestSetBit(AbsValue);	// Find most significant bit k, i.e., number of bits
-		WriteSymbol(ValueTable, ZeroValue32BitInt - (HiBit + 1));	// Code number of bits, recenter around new zero, e.g., 32 - 10 bits
-		WriteBits(AbsValue - (1 << HiBit), HiBit);	// code remaining k bits verbatim
-	}	
+	
 }
 
 void FCodecV1Encoder::SetupTables()
@@ -888,10 +734,9 @@ void FCodecV1Encoder::SetupTables()
 	// Most tables store 32-bit integers stored with a bit-length;raw value scheme. Some store specific symbols.
 	EncodingContext.ResidualIndicesTable.Initialize(HuffmanTableInt32SymbolCount);
 	EncodingContext.ResidualVertexPosTable.Initialize(HuffmanTableInt32SymbolCount);
-	EncodingContext.VertexPredictionModeTable.Initialize(VertexPredictionModeCount);
 	EncodingContext.ResidualColorTable.Initialize(HuffmanTableInt32SymbolCount);
-	EncodingContext.ResidualNormalTangentXTable.Initialize(HuffmanTableInt32SymbolCount);
-	EncodingContext.ResidualNormalTangentZTable.Initialize(HuffmanTableInt32SymbolCount);
+	EncodingContext.ResidualNormalTangentXTable.Initialize(HuffmanTableInt8SymbolCount);
+	EncodingContext.ResidualNormalTangentZTable.Initialize(HuffmanTableInt8SymbolCount);
 	EncodingContext.ResidualUVTable.Initialize(HuffmanTableInt32SymbolCount);
 	EncodingContext.ResidualMotionVectorTable.Initialize(HuffmanTableInt32SymbolCount);
 	// Add additional tables here
@@ -906,7 +751,6 @@ void FCodecV1Encoder::SetPrepass(bool bPrepass)
 	if ( !VertexInfo.bConstantIndices)
 		EncodingContext.ResidualIndicesTable.SetPrepass(bPrepass);	
 	EncodingContext.ResidualVertexPosTable.SetPrepass(bPrepass);
-	EncodingContext.VertexPredictionModeTable.SetPrepass(bPrepass);
 	if (VertexInfo.bHasColor0)
 	{
 		EncodingContext.ResidualColorTable.SetPrepass(bPrepass);
@@ -942,7 +786,6 @@ void FCodecV1Encoder::WriteTables()
 	if ( !VertexInfo.bConstantIndices)
 		EncodingContext.ResidualIndicesTable.Serialize(Writer);	
 	EncodingContext.ResidualVertexPosTable.Serialize(Writer);
-	EncodingContext.VertexPredictionModeTable.Serialize(Writer);
 	if (VertexInfo.bHasColor0)
 	{
 		EncodingContext.ResidualColorTable.Serialize(Writer);
@@ -972,28 +815,29 @@ void FCodecV1Encoder::WriteTables()
 
 FCodecV1Decoder::FCodecV1Decoder()
 {
+	// Precalculate table mapping symbol index to non-raw bits. ((Sign ? -2 : 1) << NumRawBits)
+	for (int32 NumRawBits = 1; NumRawBits <= 30; NumRawBits++)
+	{
+		for (int32 Sign = 0; Sign <= 1; Sign++)
+		{
+			HighBitsLUT[2 + Sign + NumRawBits * 2] = (Sign ? -2 : 1) << NumRawBits;
+		}
+	}	
 }
 
 void FCodecV1Decoder::DecodeIndexStream(uint32* Stream, uint64 ElementOffset, uint32 ElementCount)
 {
-	FRingBuffer<uint32, IndexStreamCodingHistorySize> LastReconstructed(IndexStreamCodingHistorySize); // History holding previously seen indices
-
 	const uint8* RawElementData = (const uint8*)Stream;
 
+	uint32 Value = 0;
 	for (uint32 ElementIdx = 0; ElementIdx < ElementCount; ++ElementIdx, RawElementData += ElementOffset)
 	{
 		// Read coded residual
 		int32 DecodedResidual = ReadInt32(DecodingContext.ResidualIndicesTable);
-
-		uint32 Prediction = LastReconstructed[0]; // Delta coding, best effort
-		uint32 Reconstructed = DecodedResidual + Prediction; // Reconstruct
+		Value += DecodedResidual;
 
 		// Save result to our list
-		uint32* Value = (uint32*)RawElementData;
-		*Value = Reconstructed;
-
-		// Store previous encountered values
-		LastReconstructed.Push(Reconstructed);
+		*(uint32*)RawElementData = Value;
 	}
 }
 
@@ -1005,7 +849,7 @@ void FCodecV1Decoder::DecodeMotionVectorStream(FVector* Stream, uint64 ElementOf
 
 	FQuantizerVector3 Quantizer(Header.QuantizationPrecision); // We quantize MVs to a certain precision just like the positions
 
-	FRingBuffer<FIntVector, MotionVectorStreamCodingHistorySize> ReconstructedHistory(MotionVectorStreamCodingHistorySize, FIntVector(0, 0, 0)); // Previously seen texture coordinates
+	FIntVector QuantizedValue(0, 0, 0);
 
 	const uint8* RawElementData = (const uint8*)Stream;
 
@@ -1017,15 +861,8 @@ void FCodecV1Decoder::DecodeMotionVectorStream(FVector* Stream, uint64 ElementOf
 		DecodedResidual.Y = ReadInt32(DecodingContext.ResidualMotionVectorTable);
 		DecodedResidual.Z = ReadInt32(DecodingContext.ResidualMotionVectorTable);
 
-		FIntVector Prediction = ReconstructedHistory[0]; // Delta coding													  
-		FIntVector Reconstructed = DecodedResidual + Prediction; // Reconstruct
-		FVector Decoded = Quantizer.Dequantize(Reconstructed); // Dequantize
-
-		FVector* Value = (FVector*)RawElementData;	// Save result to our list
-		*Value = Decoded;
-
-		// Store previous encountered values
-		ReconstructedHistory.Push(Reconstructed);
+		QuantizedValue += DecodedResidual;
+		*(FVector*)RawElementData = Quantizer.Dequantize(QuantizedValue);
 	}
 }
 
@@ -1037,7 +874,7 @@ void FCodecV1Decoder::DecodeUVStream(FVector2D* Stream, uint64 ElementOffset, ui
 
 	FQuantizerVector2 Quantizer(Header.Range, Header.QuantizationBits); // We quantize UVs to a number of bits, set in the bitstream header
 	
-	FRingBuffer<FIntVector, UVStreamCodingHistorySize> ReconstructedHistory(UVStreamCodingHistorySize, FIntVector(0, 0, 0)); // Previously seen texture coordinates
+	FIntVector QuantizedValue(0, 0, 0);
 
 	const uint8* RawElementData = (const uint8*)Stream;
 
@@ -1048,56 +885,45 @@ void FCodecV1Decoder::DecodeUVStream(FVector2D* Stream, uint64 ElementOffset, ui
 		DecodedResidual.X = ReadInt32(DecodingContext.ResidualUVTable);
 		DecodedResidual.Y = ReadInt32(DecodingContext.ResidualUVTable);
 
-		FIntVector Prediction = ReconstructedHistory[0]; // Delta coding													  
-		FIntVector Reconstructed = DecodedResidual + Prediction; // Reconstruct
-		FVector2D Decoded = Quantizer.Dequantize(Reconstructed); // Dequantize
-
-		FVector2D* Value = (FVector2D*)RawElementData;	// Save result to our list
-		*Value = Decoded;
-
-		// Store previous encountered values
-		ReconstructedHistory.Push(Reconstructed);
+		QuantizedValue += DecodedResidual;
+		*(FVector2D*)RawElementData = Quantizer.Dequantize(QuantizedValue);
 	}
 }
 
-void FCodecV1Decoder::DecodeNormalStream(FPackedNormal* Stream, uint64 ElementOffset, uint32 ElementCount, FHuffmanTable& Table)
+void FCodecV1Decoder::DecodeNormalStream(FPackedNormal* Stream, uint64 ElementOffset, uint32 ElementCount, FHuffmanDecodeTable& Table)
 {
-	FRingBuffer<FIntVector4, NormalStreamCodingHistorySize> ReconstructedHistory(NormalStreamCodingHistorySize, FIntVector4(128, 128, 128, 128)); // Previously seen normals
+	uint8 x = 128, y = 128, z = 128, w = 128;
 
 	const uint8* RawElementData = (const uint8*)Stream;
 
+	check(HUFFMAN_MAX_CODE_LENGTH * 4 <= MINIMUM_BITS_AFTER_REFILL);	// Make sure we can safely decode all 4 symbols with a single refill
+
+	FHuffmanBitStreamReader& Reader = *DecodingContext.Reader;
 	for (uint32 ElementIdx = 0; ElementIdx < ElementCount; ++ElementIdx, RawElementData += ElementOffset) // Walk normals
 	{
 		// Read coded residual
-		FIntVector4 DecodedResidual;
-		DecodedResidual.X = ReadInt32(Table);
-		DecodedResidual.Y = ReadInt32(Table);
-		DecodedResidual.Z = ReadInt32(Table);
-		DecodedResidual.W = ReadInt32(Table);
-
-		FIntVector4 Prediction = ReconstructedHistory[0]; // Delta coding
-		FIntVector4 Reconstructed = FCodecV1SharedTools::SumVector4(DecodedResidual, Prediction);
-
-		FPackedNormal* Value = (FPackedNormal*)RawElementData;	// Save result to our list
-		Value->Vector.X = Reconstructed.X;
-		Value->Vector.Y = Reconstructed.Y;
-		Value->Vector.Z = Reconstructed.Z;
-		Value->Vector.W = Reconstructed.W;
-
-		// Store previous encountered values
-		ReconstructedHistory.Push(Reconstructed);
+		Reader.Refill();
+		x += Table.DecodeNoRefill(Reader);	
+		y += Table.DecodeNoRefill(Reader);
+		z += Table.DecodeNoRefill(Reader);
+		w += Table.DecodeNoRefill(Reader);
+		
+		FPackedNormal* Value = (FPackedNormal*)RawElementData;
+		Value->Vector.X = x;
+		Value->Vector.Y = y;
+		Value->Vector.Z = z;
+		Value->Vector.W = w;
 	}
 }
 
 void FCodecV1Decoder::DecodeColorStream(FColor* Stream, uint64 ElementOffset, uint32 ElementCount)
 {
-	FRingBuffer<FIntVector4, ColorStreamCodingHistorySize> ReconstructedHistory(ColorStreamCodingHistorySize, FIntVector4(128, 128, 128, 255)); // Previously seen colors
+	FIntVector4 QuantizedValue(128, 128, 128, 255);
 
 	const uint8* RawElementData = (const uint8*)Stream;
 
 	for (uint32 ElementIdx = 0; ElementIdx < ElementCount; ++ElementIdx, RawElementData += ElementOffset)
 	{
-		FIntVector4 Prediction = ReconstructedHistory[0]; // Delta coding
 		FIntVector4 DecodedResidual(0, 0, 0, 0);
 
 		int32 SkipBit = ReadBits(1); // 1: Perfect prediction, nothing coded, 0: we have coded residuals
@@ -1111,18 +937,14 @@ void FCodecV1Decoder::DecodeColorStream(FColor* Stream, uint64 ElementOffset, ui
 			int32 DecodedResidualA = ReadInt32(DecodingContext.ResidualColorTable);
 
 			DecodedResidual = FIntVector4(DecodedResidualR, DecodedResidualG, DecodedResidualB, DecodedResidualA);
-		}	
-
-		FIntVector4 Reconstructed = FCodecV1SharedTools::SumVector4(DecodedResidual, Prediction);
+			QuantizedValue = FCodecV1SharedTools::SumVector4(QuantizedValue, DecodedResidual);
+		}
 																			
 		FColor* Value = (FColor*)RawElementData; // Save result to our list
-		Value->R = Reconstructed.X;
-		Value->G = Reconstructed.Y;
-		Value->B = Reconstructed.Z;
-		Value->A = Reconstructed.W;
-
-		// Store previous encountered values
-		ReconstructedHistory.Push(Reconstructed);
+		Value->R = QuantizedValue.X;
+		Value->G = QuantizedValue.Y;
+		Value->B = QuantizedValue.Z;
+		Value->A = QuantizedValue.W;
 	}
 }
 
@@ -1136,129 +958,39 @@ void FCodecV1Decoder::DecodePositionStream(const uint32* IndexStream, uint64 Ind
 
 	FQuantizerVector3 Quantizer(Header.QuantizationPrecision);
 	
-	FRingBuffer<uint32, VertexStreamCodingIndexHistorySize> IndexHistory(VertexStreamCodingIndexHistorySize); // Previously seen indices
-	FRingBuffer<FIntVector, VertexStreamCodingVertexHistorySize> ReconstructedVertexHistory(VertexStreamCodingVertexHistorySize, FIntVector(0, 0, 0)); // Previously seen positions
 	int64 MaxEncounteredIndex = -1; // We rely on indices being references in order, a requirement of the encoder and enforced by the preprocessor
 	uint32 DecodedVertexCount = 0;
 
 	const uint8* RawElementDataIndices = (const uint8*)IndexStream;
 	const uint8* RawElementDataVertices = (const uint8*)VertexStream;
 
-	uint32 PreviousTriangleCornerIndices[3] = { 0, 0, 0 };
-	const uint32 FaceCount = IndexElementCount / 3;
+	FIntVector QuantizedValue(0, 0, 0);
 
 	// Walk over indices/triangles
-	for (uint32 FaceIdx = 0; FaceIdx < FaceCount; ++FaceIdx, RawElementDataIndices += 3 * IndexElementOffset)
+	for (uint32 IndexIdx = 0; IndexIdx < IndexElementCount; ++IndexIdx)
 	{
-		const uint32 CornerIndices[3]
-		{
-			*(uint32*)(RawElementDataIndices),
-			*(uint32*)(RawElementDataIndices + IndexElementOffset),
-			*(uint32*)(RawElementDataIndices + 2 * IndexElementOffset)
-		};
-
-		// Find if this triangle shares an edge with the previous one, if so, get the rotation offsets to rotate the triangle's indices 
-		// We want a well-defined order of indices so we can look up the right vertices to apply the parallelogram predictor
-		uint32 PreviousTriangleRotationOffsets[3] = { 0, 0, 0 }; 
-		bool bTrianglesShareEdge = FCodecV1SharedTools::FindTrianglesSharedEdge(PreviousTriangleCornerIndices, CornerIndices, PreviousTriangleRotationOffsets);
-
-		// Walk over triangle corners
-		for (uint32 CornerIndex = 0; CornerIndex < 3; ++CornerIndex)
-		{
-			const uint32 IndexValue = CornerIndices[CornerIndex];
-			const bool NewlyEncounteredVertex = (int64)IndexValue > MaxEncounteredIndex; // We rely on indices being references in order
-			MaxEncounteredIndex = FMath::Max((int64)IndexValue, MaxEncounteredIndex);
-
-			if (NewlyEncounteredVertex)
-			{
-				checkf(DecodedVertexCount < MaxVertexElementCount, TEXT("Encountering more vertices than we have encoded. Encoding and decoding algorithms don't seem to match. Please make sure the preprocessor has processed the mesh such that vertexes are referenced in-order, i.e. if a previously unreferenced vertex is referenced by the index list it's id will always be the next unreferenced id instead of some random unused id."));
-
-				// Read prediction mode
-				const uint32 PredictionMode = ReadSymbol(DecodingContext.VertexPredictionModeTable);
-
-				// Read coded residual
-				const FIntVector DecodedResidual = { ReadInt32(DecodingContext.ResidualVertexPosTable), ReadInt32(DecodingContext.ResidualVertexPosTable), ReadInt32(DecodingContext.ResidualVertexPosTable) };
-				DecodedVertexCount++;
-
-				// Predict using the history according to the prediction mode 
-				const FIntVector Prediction = FCodecV1SharedTools::PredictVertex(CornerIndex, PreviousTriangleRotationOffsets, ReconstructedVertexHistory, PredictionMode);
-
-				// Reconstruct, dequantize
-				const FIntVector Reconstructed = DecodedResidual + Prediction;
-				// Store previous encountered values
-				ReconstructedVertexHistory.Push(Reconstructed);
-
-				// Save result to our list
-				FVector* Value = (FVector*)RawElementDataVertices;
-				*Value = Quantizer.Dequantize(Reconstructed + Header.Translation);
-				RawElementDataVertices += VertexElementOffset;		
-			}
-			else
-			{
-				// Keep our vertex history consistent with the indices we have seen.
-				// We fetch from our previously reconstructed vertex history. Our history is limited in size, so
-				// we do a explicit search and if the value is not found, we use the last seen vector as a safe prediction
-				FIntVector EncounteredVertex = ReconstructedVertexHistory[0];
-				for (uint32 Index = 0; Index < IndexHistory.Num(); ++Index)
-				{
-					if (IndexHistory[Index] == IndexValue)
-					{
-						// We have it in our history buffer
-						EncounteredVertex = ReconstructedVertexHistory[Index];
-						break;
-					}
-				}
-
-				ReconstructedVertexHistory.Push(EncounteredVertex);
-			}
-
-			IndexHistory.Push(IndexValue);
-		}
+		const uint32 IndexValue = *(const uint32*)RawElementDataIndices;
+		RawElementDataIndices += IndexElementOffset;
 
-		// Save our triangle
-		PreviousTriangleCornerIndices[0] = CornerIndices[0];
-		PreviousTriangleCornerIndices[1] = CornerIndices[1];
-		PreviousTriangleCornerIndices[2] = CornerIndices[2];
-	}
-}
-
-void FCodecV1Decoder::ZeroStream(uint8* Stream, uint64 ElementOffset, uint64 ElementSize, uint32 NumElements)
-{
-	uint8* RawElementData = (uint8*)Stream;
-
-	FMemory::Memzero(RawElementData, ElementSize * NumElements);
-
-	for (uint32 ElementIdx = 0; ElementIdx < NumElements; ++ElementIdx, RawElementData += ElementOffset)
-	{
-		FMemory::Memzero(RawElementData, ElementSize);
-	}
-}
-
-void FCodecV1Decoder::ClearStream(FVector* Stream, uint64 ElementOffset, uint32 NumElements)
-{
-	ZeroStream((uint8*)Stream, ElementOffset, sizeof(FVector), NumElements);
-}
+		if ((int64)IndexValue > MaxEncounteredIndex)
+		{
+			MaxEncounteredIndex = (int64)IndexValue;
+			checkf(DecodedVertexCount < MaxVertexElementCount, TEXT("Encountering more vertices than we have encoded. Encoding and decoding algorithms don't seem to match. Please make sure the preprocessor has processed the mesh such that vertexes are referenced in-order, i.e. if a previously unreferenced vertex is referenced by the index list it's id will always be the next unreferenced id instead of some random unused id."));
 
-void FCodecV1Decoder::ClearStream(FColor* Stream, uint64 ElementOffset, uint32 NumElements)
-{
-	ZeroStream((uint8*)Stream, ElementOffset, sizeof(FColor), NumElements);
-}
+			// Read coded residual
+			const FIntVector DecodedResidual = { ReadInt32(DecodingContext.ResidualVertexPosTable), ReadInt32(DecodingContext.ResidualVertexPosTable), ReadInt32(DecodingContext.ResidualVertexPosTable) };
+			DecodedVertexCount++;
 
-void FCodecV1Decoder::ClearStream(FPackedNormal* Stream, uint64 ElementOffset, uint32 NumElements)
-{
-	uint8* RawElementData = (uint8*)Stream;
+			QuantizedValue += DecodedResidual;
 
-	for (uint32 ElementIdx = 0; ElementIdx < NumElements; ++ElementIdx, RawElementData += ElementOffset)
-	{
-		new (RawElementData) FPackedNormal(); //Zero initialize
+			// Save result to our list
+			FVector* Value = (FVector*)RawElementDataVertices;
+			*Value = Quantizer.Dequantize(QuantizedValue + Header.Translation);
+			RawElementDataVertices += VertexElementOffset;		
+		}
 	}
 }
 
-void FCodecV1Decoder::ClearStream(FVector2D* Stream, uint64 ElementOffset, uint32 NumElements)
-{
-	ZeroStream((uint8*)Stream, ElementOffset, sizeof(FVector2D), NumElements);
-}
-
 void FCodecV1Decoder::SetupAndReadTables()
 {
 	// Initialize and read Huffman tables from the bitstream
@@ -1268,7 +1000,6 @@ void FCodecV1Decoder::SetupAndReadTables()
 	if (!VertexInfo.bConstantIndices)
 		DecodingContext.ResidualIndicesTable.Initialize(Reader);	
 	DecodingContext.ResidualVertexPosTable.Initialize(Reader);
-	DecodingContext.VertexPredictionModeTable.Initialize(Reader);
 	
 	if (VertexInfo.bHasColor0)
 	{
@@ -1310,6 +1041,15 @@ void FCodecV1Decoder::ReadCodedStreamDescription()
 	VertexInfo.bConstantIndices = (ReadBits(1) == 1);
 }
 
+DECLARE_CYCLE_STAT(TEXT("SetupAndReadTables"), STAT_SetupAndReadTables, STATGROUP_GeometryCache);
+DECLARE_CYCLE_STAT(TEXT("DecodeIndexStream"), STAT_DecodeIndexStream, STATGROUP_GeometryCache);
+DECLARE_CYCLE_STAT(TEXT("DecodePositionStream"), STAT_DecodePositionStream, STATGROUP_GeometryCache);
+DECLARE_CYCLE_STAT(TEXT("DecodeColorStream"), STAT_DecodeColorStream, STATGROUP_GeometryCache);
+DECLARE_CYCLE_STAT(TEXT("DecodeTangentXStream"), STAT_DecodeTangentXStream, STATGROUP_GeometryCache);
+DECLARE_CYCLE_STAT(TEXT("DecodeTangentZStream"), STAT_DecodeTangentZStream, STATGROUP_GeometryCache);
+DECLARE_CYCLE_STAT(TEXT("DecodeUVStream"), STAT_DecodeUVStream, STATGROUP_GeometryCache);
+
+
 bool FCodecV1Decoder::DecodeFrameData(FBufferReader& Reader, FGeometryCacheMeshData &OutMeshData)
 {
 	FExperimentTimer DecodingTime;
@@ -1333,17 +1073,20 @@ bool FCodecV1Decoder::DecodeFrameData(FBufferReader& Reader, FGeometryCacheMeshD
 	
 	// Read the payload in memory to pass to the bit reader
 	TArray<uint8> Bytes;
-	Bytes.AddUninitialized(Header.PayloadSize);
-	Reader.Serialize(Bytes.GetData(), Bytes.Num()); 	
-	FHuffmanBitStreamReader BitReader(Bytes); // This copies the data again. We should optimize this out.
-	Bytes.Empty(); // No need anymore
+	Bytes.AddUninitialized(Header.PayloadSize + 16);	// Overallocate by 16 bytes to ensure BitReader can safely perform uint64 reads.
+	Reader.Serialize(Bytes.GetData(), Header.PayloadSize);
+	FHuffmanBitStreamReader BitReader(Bytes.GetData(), Bytes.Num());
 	DecodingContext.Reader = &BitReader;
 
 	// Read which vertex attributes are in the bit stream
 	ReadCodedStreamDescription();
 
 	// Restore entropy coding contexts
-	SetupAndReadTables();
+	{
+		SCOPE_CYCLE_COUNTER(STAT_SetupAndReadTables);
+		SetupAndReadTables();
+	}
+	
 	
 	{
 		// Decode streams
@@ -1353,19 +1096,24 @@ bool FCodecV1Decoder::DecodeFrameData(FBufferReader& Reader, FGeometryCacheMeshD
 		{
 			OutMeshData.Indices.Empty(Header.IndexCount);
 			OutMeshData.Indices.AddUninitialized(Header.IndexCount);
+			SCOPE_CYCLE_COUNTER(STAT_DecodeIndexStream);
 			DecodeIndexStream(&Indices[0], Indices.GetTypeSize(), Header.IndexCount);
 		}
 		
 		TArray<FVector>& Positions = OutMeshData.Positions;
 		OutMeshData.Positions.Empty(Header.VertexCount);
 		OutMeshData.Positions.AddZeroed(Header.VertexCount);
-		DecodePositionStream(&Indices[0], Indices.GetTypeSize(), Indices.Num(), &Positions[0], Positions.GetTypeSize(), Header.VertexCount);
+		{
+			SCOPE_CYCLE_COUNTER(STAT_DecodePositionStream);
+			DecodePositionStream(&Indices[0], Indices.GetTypeSize(), Indices.Num(), &Positions[0], Positions.GetTypeSize(), Header.VertexCount);
+		}
 
 		OutMeshData.Colors.Empty(Header.VertexCount);
 		OutMeshData.Colors.AddZeroed(Header.VertexCount);
 		if (OutMeshData.VertexInfo.bHasColor0)
 		{
 			TArray<FColor>& Colors = OutMeshData.Colors;
+			SCOPE_CYCLE_COUNTER(STAT_DecodeColorStream);
 			DecodeColorStream(&Colors[0], Colors.GetTypeSize(), Header.VertexCount);
 		}
 
@@ -1374,6 +1122,7 @@ bool FCodecV1Decoder::DecodeFrameData(FBufferReader& Reader, FGeometryCacheMeshD
 		if (OutMeshData.VertexInfo.bHasTangentX)
 		{
 			TArray<FPackedNormal>& TangentsX = OutMeshData.TangentsX;
+			SCOPE_CYCLE_COUNTER(STAT_DecodeTangentXStream);
 			DecodeNormalStream(&TangentsX[0], TangentsX.GetTypeSize(), Header.VertexCount, DecodingContext.ResidualNormalTangentXTable);
 		}
 
@@ -1382,6 +1131,7 @@ bool FCodecV1Decoder::DecodeFrameData(FBufferReader& Reader, FGeometryCacheMeshD
 		if (OutMeshData.VertexInfo.bHasTangentZ)
 		{
 			TArray<FPackedNormal>& TangentsZ = OutMeshData.TangentsZ;
+			SCOPE_CYCLE_COUNTER(STAT_DecodeTangentZStream);
 			DecodeNormalStream(&TangentsZ[0], TangentsZ.GetTypeSize(), Header.VertexCount, DecodingContext.ResidualNormalTangentZTable);
 		}
 
@@ -1390,6 +1140,7 @@ bool FCodecV1Decoder::DecodeFrameData(FBufferReader& Reader, FGeometryCacheMeshD
 		if (OutMeshData.VertexInfo.bHasUV0)
 		{
 			TArray<FVector2D>& TextureCoordinates = OutMeshData.TextureCoordinates;
+			SCOPE_CYCLE_COUNTER(STAT_DecodeUVStream);
 			DecodeUVStream(&TextureCoordinates[0], TextureCoordinates.GetTypeSize(), Header.VertexCount);
 		}
 
@@ -1407,6 +1158,7 @@ bool FCodecV1Decoder::DecodeFrameData(FBufferReader& Reader, FGeometryCacheMeshD
 			UE_LOG(LogGeoCaStreamingCodecV1, Log, TEXT("Decoded frame with %i vertices in %.2f milliseconds."), Positions.Num(), TimeFloat);
 		}
 	}
+	DecodingContext.Reader = NULL;
 
 	return true;
 }
@@ -1422,24 +1174,19 @@ void FCodecV1Decoder::ReadBytes(void* Data, uint32 NumBytes)
 	}
 }
 
-int32 FCodecV1Decoder::ReadInt32(FHuffmanTable& ValueTable)
-{	
-	// We code the length of the value (huffman table), followed by the raw bits. Small values take up very small amount of bits.
-	const int32 ZeroValue32BitInt = 32;	// Zero value used to shift interval [32,31] to [0,63]
-
-	const int32 PackedLength = ReadSymbol(ValueTable); // Read number of bits, i.e., most significant bit + 1
-	if (PackedLength == ZeroValue32BitInt)
+int32 FCodecV1Decoder::ReadInt32(FHuffmanDecodeTable& ValueTable)
+{
+	// See write WriteInt32 for encoding details.
+	const int32 Packed = ReadSymbol(ValueTable);
+	if (Packed < 4)
 	{
-		return 0;
+		// [-2, 1] coded directly with no additional raw bits
+		return Packed - 2;
 	}
-	else if (PackedLength > ZeroValue32BitInt)
-	{	
-		int32 NumBits = PackedLength - 1 - ZeroValue32BitInt;  // Positive number, undo re centering around new zero, extract length
-		return ReadBits(NumBits) + (1 << NumBits);	// Read remaining bits
+	else
+	{
+		// At least one raw bit.
+		int32 NumRawBits = (Packed - 2) >> 1;
+		return ReadBitsNoRefill(NumRawBits) + HighBitsLUT[Packed];
 	}
-	else /*if (PackedLength < ZeroValue32BitInt)*/
-	{		
-		int32 NumBits = ZeroValue32BitInt - PackedLength - 1;	 // Negative number, undo re centering around new zero, extract length
-		return -(ReadBits(NumBits) + (1 << NumBits)); // Read remaining bits
-	}	
 }
\ No newline at end of file
diff --git a/Engine/Plugins/Experimental/GeometryCache/Source/GeometryCache/Private/CodecV1.h b/Engine/Plugins/Experimental/GeometryCache/Source/GeometryCache/Private/CodecV1.h
--- a/Engine/Plugins/Experimental/GeometryCache/Source/GeometryCache/Private/CodecV1.h
+++ b/Engine/Plugins/Experimental/GeometryCache/Source/GeometryCache/Private/CodecV1.h
@@ -87,14 +87,13 @@ struct FCodecV1DecodingContext
 	/** Reader to read bit stream from */
 	FHuffmanBitStreamReader* Reader;
 	/** Huffman tables */
-	FHuffmanTable ResidualIndicesTable;	
-	FHuffmanTable ResidualVertexPosTable;
-	FHuffmanTable ResidualColorTable;
-	FHuffmanTable ResidualNormalTangentXTable;
-	FHuffmanTable ResidualNormalTangentZTable;
-	FHuffmanTable ResidualUVTable;
-	FHuffmanTable VertexPredictionModeTable;
-	FHuffmanTable ResidualMotionVectorTable;	
+	FHuffmanDecodeTable ResidualIndicesTable;	
+	FHuffmanDecodeTable ResidualVertexPosTable;
+	FHuffmanDecodeTable ResidualColorTable;
+	FHuffmanDecodeTable ResidualNormalTangentXTable;
+	FHuffmanDecodeTable ResidualNormalTangentZTable;
+	FHuffmanDecodeTable ResidualUVTable;
+	FHuffmanDecodeTable ResidualMotionVectorTable;
 };
 
 /** Configuration settings for the encoder */
@@ -131,19 +130,18 @@ struct FCodecV1EncodingContext
 	/** Huffman bit writer to write bitstream to */
 	FHuffmanBitStreamWriter* Writer;	
 	/** Huffman tables */
-	FHuffmanTable ResidualIndicesTable;
-	FHuffmanTable ResidualVertexPosTable;
-	FHuffmanTable ResidualColorTable;
-	FHuffmanTable ResidualColorSkipTable;
-	FHuffmanTable ResidualNormalTangentXTable;
-	FHuffmanTable ResidualNormalTangentZTable;
-	FHuffmanTable ResidualUVTable;
-	FHuffmanTable VertexPredictionModeTable;
-	FHuffmanTable ResidualMotionVectorTable;
+	FHuffmanEncodeTable ResidualIndicesTable;
+	FHuffmanEncodeTable ResidualVertexPosTable;
+	FHuffmanEncodeTable ResidualColorTable;
+	FHuffmanEncodeTable ResidualColorSkipTable;
+	FHuffmanEncodeTable ResidualNormalTangentXTable;
+	FHuffmanEncodeTable ResidualNormalTangentZTable;
+	FHuffmanEncodeTable ResidualUVTable;
+	FHuffmanEncodeTable ResidualMotionVectorTable;
 };
 
-const uint32 VertexPredictionModeCount = 4; // vertex position prediction methods
 const int32 HuffmanTableInt32SymbolCount = 64; // 31 negative lengths, zero, 32 positive length
+const int32 HuffmanTableInt8SymbolCount = 256;
 const uint32 VertexStreamCodingIndexHistorySize = 9; // Sizes of the previously-seen histories used for prediction of the various stream elements
 const uint32 VertexStreamCodingVertexHistorySize = 9;
 const uint32 IndexStreamCodingHistorySize = 5;
@@ -156,17 +154,6 @@ const uint32 MotionVectorStreamCodingHistorySize = 9;
 class FCodecV1SharedTools 
 {
 public:
-	
-	/** Predict a vertex according to the prediction mode using the vertex history	*/
-	static FORCEINLINE FIntVector PredictVertex(uint32 CornerIndex, uint32 PreviousTriangleRotationOffsets[3], FRingBuffer<FIntVector, VertexStreamCodingVertexHistorySize>& ReconstructedVertexHistory, uint32 Mode);
-	
-	/** Find the edge two triangles share. Return the rotation of the first triangle (index offsets) such that the first 
-	    two indices point to the shared edge: (shared corner index 1, shared corner index 2, non-shared corner index) */
-	static FORCEINLINE bool FindTrianglesSharedEdge(const uint32 Triangle1[3], const uint32 Triangle2[3], uint32 orderedTriangle1[3]);
-	
-	/** Estimate the coding cost of a vector */
-	static FORCEINLINE int32 EstimateCodingCost(const FIntVector& Value);
-
 	/** Sum two IntVector4 vectors, because IntVector4 does not implement arithmetic operations contrary to IntVector */
 	static FORCEINLINE FIntVector4 SumVector4(const FIntVector4& First, const FIntVector4& Second)
 	{
@@ -222,7 +209,7 @@ private:
 	/** Encode a buffer with vertex colors */
 	void EncodeColorStream(const FColor* Stream, uint64 ElementOffsetBytes, uint32 ElementCount, FStreamEncodingStatistics& Stats);
 	/** Encode a buffer with vertex normals */
-	void EncodeNormalStream(const FPackedNormal* Stream, uint64 ElementOffsetBytes, uint32 ElementCount, FHuffmanTable& Table, FStreamEncodingStatistics& Stats);
+	void EncodeNormalStream(const FPackedNormal* Stream, uint64 ElementOffsetBytes, uint32 ElementCount, FHuffmanEncodeTable& Table, FStreamEncodingStatistics& Stats);
 	/** Encode a buffer with vertex texture coordinates */
 	void EncodeUVStream(const FVector2D* Stream, uint64 ElementOffsetBytes, uint32 ElementCount, FStreamEncodingStatistics& Stats);
 	/** Encode a buffer with vertex motion vectors */
@@ -242,7 +229,7 @@ private:
 	void WriteBytes(const void* Data, int64 Num);
 	
 	/**	Write a 32-bit signed value from the bit stream using a Huffman coder */
-	void WriteInt32(FHuffmanTable& ValueTable, int32 Value);
+	void WriteInt32(FHuffmanEncodeTable& ValueTable, int32 Value);
 
 	/** Write a given number of bits to the bit stream using the Huffman bit writer */
 	FORCEINLINE void WriteBits(int32 Data, uint32 NumBits)
@@ -256,7 +243,7 @@ private:
 	}
 
 	/** Write a symbol to the bit stream using a Huffman Table */
-	FORCEINLINE void WriteSymbol(FHuffmanTable& table, int32 symbol)
+	FORCEINLINE void WriteSymbol(FHuffmanEncodeTable& table, int32 symbol)
 	{
 		// Write using the table, will not write anything in the prepass phase
 		table.Encode(*EncodingContext.Writer, symbol);
@@ -320,22 +307,12 @@ private:
 	/** Decode a buffer of vertex colors from the bitstream */
 	void DecodeColorStream(FColor* Stream, uint64 ElementOffset, uint32 ElementCount);
 	/** Decode a buffer of vertex normals from the bitstream */
-	void DecodeNormalStream(FPackedNormal* Stream, uint64 ElementOffset, uint32 ElementCount, FHuffmanTable& Table);
+	void DecodeNormalStream(FPackedNormal* Stream, uint64 ElementOffset, uint32 ElementCount, FHuffmanDecodeTable& Table);
 	/** Decode a buffer of vertex texture coordinates from the bitstream */
 	void DecodeUVStream(FVector2D* Stream, uint64 ElementOffset, uint32 ElementCount);
 	/** Decode a buffer of vertex motion vectors from the bitstream */
 	void DecodeMotionVectorStream(FVector* Stream, uint64 ElementOffset, uint32 ElementCount);
-	/** Clear a stream of vectors */
-	void ClearStream(FVector* Stream, uint64 ElementOffset, uint32 NumElements);
-	/** Clear a stream of colors */
-	void ClearStream(FColor* Stream, uint64 ElementOffset, uint32 NumElements);
-	/** Clear a stream of normals */
-	void ClearStream(FPackedNormal* Stream, uint64 ElementOffset, uint32 NumElements);	
-	/** Clear a stream of texture coordinates*/
-	void ClearStream(FVector2D* Stream, uint64 ElementOffset, uint32 NumElements);
-	/** Memclear a generic stream */
-	void ZeroStream(uint8* Stream, uint64 ElementOffset, uint64 ElementSize, uint32 NumElements);
-
+	
 	/** Initialize and read Huffman tables from the bitstream */
 	void SetupAndReadTables();	
 	
@@ -343,24 +320,32 @@ private:
 	FORCEINLINE void ReadBytes(void* Data, uint32 NumBytes);
 
 	/**	Read a 32-bit signed value from the bit stream using a Huffman coder */
-	FORCEINLINE int32 ReadInt32(FHuffmanTable& ValueTable);
+	FORCEINLINE int32 ReadInt32(FHuffmanDecodeTable& ValueTable);
 
 	/** Read a symbol from the bit stream using a Huffman Table */
-	FORCEINLINE int32 ReadSymbol(FHuffmanTable& table)
+	FORCEINLINE int32 ReadSymbol(FHuffmanDecodeTable& table)
 	{
 		return table.Decode(*DecodingContext.Reader);
 	}
 
-	/** Read a given number of bits from the bit stream using the Huffman bit writer */
+	/** Read a given number of bits from the bit stream */
 	FORCEINLINE int32 ReadBits(uint32 NumBits)
 	{
 		return DecodingContext.Reader->Read(NumBits);
 	}
 
+	/** Read a given number of bits from the bit stream without refilling bits */
+	FORCEINLINE int32 ReadBitsNoRefill(uint32 NumBits)
+	{
+		return DecodingContext.Reader->ReadNoRefill(NumBits);
+	}
+
 	/** Read info on the available streams from the bit stream */
 	void ReadCodedStreamDescription();
 
-	/** Any context information to decode a frame in a sequence of frames */
+	/** Any context information to decode a frame in a sequence of frames, such as the bit steam and any Huffman tables used.*/
 	FCodecV1DecodingContext DecodingContext;	
+
+	int32 HighBitsLUT[64];
 };
 
diff --git a/Engine/Plugins/Experimental/GeometryCache/Source/GeometryCache/Private/FloatCompression.cpp b/Engine/Plugins/Experimental/GeometryCache/Source/GeometryCache/Private/FloatCompression.cpp
--- a/Engine/Plugins/Experimental/GeometryCache/Source/GeometryCache/Private/FloatCompression.cpp
+++ b/Engine/Plugins/Experimental/GeometryCache/Source/GeometryCache/Private/FloatCompression.cpp
@@ -12,6 +12,17 @@ DEFINE_LOG_CATEGORY_STATIC(LogBlankProgram, Log, All);
 
 //IMPLEMENT_APPLICATION(BlankProgram, "BlankProgram");
 
+/**
+	Move data to an 16 byte larger array to make it safe to use with HuffmanBitStreamReader
+*/
+static TArray<uint8> PadData(const TArray<uint8>& Data)
+{
+	TArray<uint8> PaddedData;
+	PaddedData.SetNumUninitialized(Data.Num() + 16);
+	FMemory::Memcpy(&PaddedData[0], Data.GetData(), Data.Num());
+	return PaddedData;
+}
+
 /**
 	Floating point compression/decompression class.
 	This is a quick sample class could probably have a nicer interface where the
@@ -30,7 +41,7 @@ public:
 
 	FloatCoder()
 	{
-		BitlengthsTable.Initialize(256);
+		BitlengthsEncodeTable.Initialize(256);
 	}
 
 	/**
@@ -38,9 +49,9 @@ public:
 	*/
 	void Encode(float *Values, int32 NumValues)
 	{
-		BitlengthsTable.SetPrepass(true);
+		BitlengthsEncodeTable.SetPrepass(true);
 		EncodePass(Values, NumValues);
-		BitlengthsTable.SetPrepass(false);
+		BitlengthsEncodeTable.SetPrepass(false);
 		EncodePass(Values, NumValues);
 
 		uint32 bytes = Stream.GetNumBytes();
@@ -58,13 +69,13 @@ public:
 	*/
 	void Decode(TArray<uint8> Data, float *Values, int32 NumValues)
 	{
-		FHuffmanBitStreamReader InputStream(Data);
+		FHuffmanBitStreamReader InputStream(Data.GetData(), Data.Num());
 
 		for (int32 Index = 0; Index < NumValues; Index++)
 		{
 			float Pred = (Index > 0) ? Values[Index - 1] : 0.0f;
 			uint32 PredInt = IntEncode(Pred);
-			uint32 PackedK = BitlengthsTable.Decode(InputStream);
+			uint32 PackedK = BitlengthsDecodeTable.Decode(InputStream);
 			uint32 Value = 0;
 			if (PackedK == Zero)
 			{
@@ -131,7 +142,8 @@ protected:
 		return *((float *)(&Output));
 	}
 
-	FHuffmanTable BitlengthsTable; //Huffman table for the most significant bit & sign combo.
+	FHuffmanEncodeTable BitlengthsEncodeTable; //Huffman table for the most significant bit & sign combo.
+	FHuffmanDecodeTable BitlengthsDecodeTable;
 
 	FHuffmanBitStreamWriter Stream; // Used to write the encoded data to
 	const uint32 Zero = 33; // perfect prediction symbol
@@ -165,7 +177,7 @@ protected:
 				// case 1: under prediction
 				uint32 Delta = CodedValue - PredValue; // absolute difference
 				uint32 HiBit = HighestSetBit(Delta); // find most significant bit k
-				BitlengthsTable.Encode(Stream, Zero + (HiBit + 1));
+				BitlengthsEncodeTable.Encode(Stream, Zero + (HiBit + 1));
 				Stream.Write(Delta - (1 << HiBit), HiBit);// code remaining k bits verbatim
 			}
 			else if (PredValue > CodedValue)
@@ -173,13 +185,13 @@ protected:
 				// case 2: over prediction
 				uint32 Delta = PredValue - CodedValue; // absolute difference
 				uint32 HiBit = HighestSetBit(Delta); // find most significant bit k
-				BitlengthsTable.Encode(Stream, Zero - (HiBit + 1)); // entropy code k
+				BitlengthsEncodeTable.Encode(Stream, Zero - (HiBit + 1)); // entropy code k
 				Stream.Write(Delta - (1 << HiBit), HiBit); // code remaining k bits verbatim
 			}
 			else
 			{
 				// case 3: perfect prediction
-				BitlengthsTable.Encode(Stream, Zero);
+				BitlengthsEncodeTable.Encode(Stream, Zero);
 			}
 		}
 		Stream.Close();
@@ -202,14 +214,16 @@ void BitstreamTest()
 
 	// Single read call
 	{
-		FHuffmanBitStreamReader TableStreamReader(TableStream.GetBytes());
+		TArray<uint8> PaddedTableData = PadData(TableStream.GetBytes());
+		FHuffmanBitStreamReader TableStreamReader(PaddedTableData.GetData(), PaddedTableData.Num());
 		uint32 Result = TableStreamReader.Read(4);
 		check(Result == 11);
 	}
 
 	// Single bit reads
 	{
-		FHuffmanBitStreamReader TableStreamReader(TableStream.GetBytes());
+		TArray<uint8> PaddedTableData = PadData(TableStream.GetBytes());
+		FHuffmanBitStreamReader TableStreamReader(PaddedTableData.GetData(), PaddedTableData.Num());
 		uint32 Result = TableStreamReader.Read();
 		check(Result == 1);
 		Result = TableStreamReader.Read();
@@ -222,7 +236,8 @@ void BitstreamTest()
 
 	// Sequential bit reads + shifting
 	{
-		FHuffmanBitStreamReader TableStreamReader(TableStream.GetBytes());
+		TArray<uint8> PaddedTableData = PadData(TableStream.GetBytes());
+		FHuffmanBitStreamReader TableStreamReader(PaddedTableData.GetData(), PaddedTableData.Num());
 		uint32 Result = TableStreamReader.Read(3);
 		Result = Result << 1;
 		Result |= TableStreamReader.Read();
@@ -231,8 +246,9 @@ void BitstreamTest()
 
 	// Peeking with zero padding
 	{
-		FHuffmanBitStreamReader TableStreamReader(TableStream.GetBytes());
-		uint32 Result = TableStreamReader.PeekWithZeros(8);
+		TArray<uint8> PaddedTableData = PadData(TableStream.GetBytes());
+		FHuffmanBitStreamReader TableStreamReader(PaddedTableData.GetData(), PaddedTableData.Num());
+		uint32 Result = TableStreamReader.Peek(8);
 		check(Result == 0xB0);
 	}
 
@@ -247,7 +263,8 @@ void BitstreamTest()
 	TableStream.Close();
 
 	{
-		FHuffmanBitStreamReader TableStreamReader(TableStream.GetBytes());
+		TArray<uint8> PaddedTableData = PadData(TableStream.GetBytes());
+		FHuffmanBitStreamReader TableStreamReader(PaddedTableData.GetData(), PaddedTableData.Num());
 		uint32 Result = TableStreamReader.Read(4);
 		check(Result == 11);
 		Result = TableStreamReader.Read(17);
@@ -277,7 +294,7 @@ void IntegerTest()
 		IntValues[Index] = (uint32)((sf*0.5 + 0.5) * 3999);
 	}
 
-	FHuffmanTable Tab; // 12 bit numbers
+	FHuffmanEncodeTable Tab; // 12 bit numbers
 	Tab.Initialize(4000);
 	FHuffmanBitStreamWriter Writer;
 	Tab.SetPrepass(true);
@@ -292,15 +309,17 @@ void IntegerTest()
 	}
 	Writer.Close();
 
-	FHuffmanBitStreamReader Reader(Writer.GetBytes());
+	TArray<uint8> PaddedWriterData = PadData(Writer.GetBytes());
+	FHuffmanBitStreamReader Reader(PaddedWriterData.GetData(), PaddedWriterData.Num());
 
 	//Serialize and deserialize the table. This is not needed as the same instance can be used but that way
 	//we can test the table serialization also.
 	FHuffmanBitStreamWriter TableStream;
 	Tab.Serialize(TableStream);
 	TableStream.Close();
-	FHuffmanBitStreamReader TableStreamReader(TableStream.GetBytes());
-	FHuffmanTable ReadTable;
+	TArray<uint8> PaddedTableData = PadData(TableStream.GetBytes());
+	FHuffmanBitStreamReader TableStreamReader(PaddedTableData.GetData(), PaddedTableData.Num());
+	FHuffmanDecodeTable ReadTable;
 	ReadTable.Initialize(TableStreamReader);
 
 	for (int32 Index = 0; Index < NumValues; Index++)
@@ -337,7 +356,8 @@ void FloatTest()
 	FloatCoder<0> coder;
 
 	coder.Encode(Values, NumValues);
-	coder.Decode(coder.GetData(), OutValues, NumValues);
+	TArray<uint8> PaddedData = PadData(coder.GetData());
+	coder.Decode(PaddedData, OutValues, NumValues);
 
 	for (int32 Index = 0; Index < NumValues; Index++)
 	{
diff --git a/Engine/Plugins/Experimental/GeometryCache/Source/GeometryCache/Private/GeometryCache.cpp b/Engine/Plugins/Experimental/GeometryCache/Source/GeometryCache/Private/GeometryCache.cpp
--- a/Engine/Plugins/Experimental/GeometryCache/Source/GeometryCache/Private/GeometryCache.cpp
+++ b/Engine/Plugins/Experimental/GeometryCache/Source/GeometryCache/Private/GeometryCache.cpp
@@ -38,10 +38,29 @@ void UGeometryCache::PostInitProperties()
 void UGeometryCache::Serialize(FArchive& Ar)
 {
 	Ar.UsingCustomVersion(FAnimPhysObjectVersion::GUID);
+	Ar.UsingCustomVersion(FRenderingObjectVersion::GUID);
+
+	auto ShowDeprecationNotification = [this]()
+	{
+		Tracks.Empty();
+		Materials.Empty();
+
+		const FText ErrorText = LOCTEXT("GeometryCacheEmptied", "Geometry Cache asset has been emptied as it does not support backwards compatibility");
+		FNotificationInfo Info(ErrorText);
+		Info.ExpireDuration = 5.0f;
+		FSlateNotificationManager::Get().AddNotification(Info);
+
+		UE_LOG(LogGeometryCache, Warning, TEXT("(%s) %s"), *ErrorText.ToString(), *GetName());
+	};
 		
 	if (Ar.CustomVer(FAnimPhysObjectVersion::GUID) >= FAnimPhysObjectVersion::GeometryCacheAssetDeprecation)
 	{
 		Super::Serialize(Ar);
+
+		if (Ar.CustomVer(FRenderingObjectVersion::GUID) < FRenderingObjectVersion::GeometryCacheFastDecoder)
+		{
+			ShowDeprecationNotification();
+		}
 	}
 	else
 	{
@@ -63,15 +82,7 @@ void UGeometryCache::Serialize(FArchive& Ar)
 			Ar << Materials;
 		}
 
-		Tracks.Empty();
-		Materials.Empty();
-		
-		const FText ErrorText = LOCTEXT("GeometryCacheEmptied", "Geometry Cache asset has been emptied as it does not support backwards compatibility");
-		FNotificationInfo Info(ErrorText);
-		Info.ExpireDuration = 5.0f;
-		FSlateNotificationManager::Get().AddNotification(Info);
-
-		UE_LOG(LogGeometryCache, Warning, TEXT("(%s) %s"), *ErrorText.ToString(), *GetName());
+		ShowDeprecationNotification();
 	}
 }
 
diff --git a/Engine/Plugins/Experimental/GeometryCache/Source/GeometryCache/Private/HuffmanBitStream.h b/Engine/Plugins/Experimental/GeometryCache/Source/GeometryCache/Private/HuffmanBitStream.h
--- a/Engine/Plugins/Experimental/GeometryCache/Source/GeometryCache/Private/HuffmanBitStream.h
+++ b/Engine/Plugins/Experimental/GeometryCache/Source/GeometryCache/Private/HuffmanBitStream.h
@@ -4,11 +4,12 @@
 #include "CoreMinimal.h"
 #include "Containers/Array.h"
 
+#define USE_UNALIGNED_READ (PLATFORM_WINDOWS | PLATFORM_MAC | PLATFORM_XBOXONE | PLATFORM_PS4)	// Little-endian platforms that support fast unaligned reads
+#define MINIMUM_BITS_AFTER_REFILL 56															// Minimum number of bits guaranteed to be available in the internal buffer after a buffer refill.
 
 /**
 A bit stream writer class for use with the Huffman coding.
-This class allows coding arbitrary sized integers up to 32 bits in size. The bits in the int are logically written
-into the stream with their most significant bit first.
+This class allows coding arbitrary sized integers up to 32 bits in size. The bits are written in Little-endian order.
 */
 class FHuffmanBitStreamWriter
 {
@@ -28,56 +29,30 @@ public:
 	void Clear()
 	{
 		Bytes.Empty();
-		Pos = -1;
-		BitsLeftInLastByte = 0;
 		bFlushed = false;
+		BitBuffer = 0;
+		BitBufferBits = 0;
 		NumBits = 0;
 	}
 
 	/**
-	Return a new integer which only the lowest NumBits of the input integer K.
-	*/
-	FORCEINLINE static uint32 LeastNBits(uint32 K, uint32 NumBits)
-	{
-		return K & ((1 << NumBits) - 1);
-	}
-
-	/**
-	Return a new integer with NumBits bits from K starting at FirstBit. The
-	returned value is right aligned so FirstBit will be in position 0 of
-	the returned number.
-	*/
-	FORCEINLINE static uint32 ExtractBits(uint32 K, int32 FirstBit, uint32 InNumBits)
-	{
-		return LeastNBits(K >> FirstBit, InNumBits);
-	}
-
-	/**
-	Write an NumBits integer value to the stream. The most significant bit
-	will be logically first in the stream.
+	Write an NumBits integer value to the stream.
 	*/
 	void Write(uint32 Bits, uint32 InNumBits)
 	{
 		checkf(bFlushed == false, TEXT("You cannot write to a stream that has already been closed"));
+		check(Bits < (1u << InNumBits));
 
-		uint32 BitsLeft = InNumBits;
-		while (BitsLeft > 0)
-		{
-			// Alloc another byte if we're out of space
-			if (BitsLeftInLastByte == 0)
-			{
-				Bytes.Add(0);
-				Pos++;
-				BitsLeftInLastByte = 8;
-			}
+		BitBuffer |= (uint64)Bits << BitBufferBits;
+		BitBufferBits += InNumBits;
+		NumBits += InNumBits;
 
-			const uint32 NumBitsToWrite = FMath::Min(BitsLeftInLastByte, BitsLeft);
-			Bytes[Pos] = Bytes[Pos] << NumBitsToWrite;
-			Bytes[Pos] |= ExtractBits(Bits, BitsLeft - NumBitsToWrite, NumBitsToWrite);
-			BitsLeft -= NumBitsToWrite;
-			BitsLeftInLastByte -= NumBitsToWrite;
+		while(BitBufferBits >= 8)
+		{
+			Bytes.Add((uint8)BitBuffer);
+			BitBuffer >>= 8;
+			BitBufferBits -= 8;
 		}
-		NumBits += InNumBits;
 	}
 
 	/**
@@ -87,9 +62,13 @@ public:
 	*/
 	void Close()
 	{
-		// Write zeros to the last byte this ensures the MSB is correctly shifted left
-		// in the last byte.
-		Write(0, BitsLeftInLastByte);
+		// Round up to next byte by appending 0s.
+		if(BitBufferBits)
+		{
+			Write(0, 8 - BitBufferBits);
+		}
+		check(BitBufferBits == 0);
+		
 		bFlushed = true;
 	}
 
@@ -122,8 +101,8 @@ public:
 
 private:
 	TArray<uint8> Bytes;
-	uint32 Pos;
-	uint32 BitsLeftInLastByte;
+	uint64 BitBuffer;
+	uint32 BitBufferBits;
 	uint32 NumBits;
 	bool bFlushed;
 };
@@ -157,35 +136,51 @@ private:
 class FHuffmanBitStreamReader
 {
 public:
-
 	/**
-	Initialize the stream for reading with the specified data bytes. Bytes will be read assuming the most significant bit in the byte comes logically
-	earlier in the stream.
+	Initialize the stream for reading with the specified data bytes. The bitstream doesn't own the data, so the caller is responsible for keeping the data valid while reader is active.
+	Additionally the caller is responsible for over-allocating the buffer by 16 bytes, so we can safely read the data as uint64s.
+	16 bytes and not 8 bytes, because the input pointer can move 8 bytes past the data and perform a 8 byte read from there.
 	*/
-	FHuffmanBitStreamReader(uint8 *Data, uint32 NumBytes)
+	FHuffmanBitStreamReader(const uint8* InBytes, uint32 InNumBytes)
 	{
-		Bytes.SetNumUninitialized(NumBytes);
-		FMemory::Memcpy(&Bytes[0], Data, NumBytes);
+		Bytes = InBytes;
+		NumBytes = InNumBytes;
 		Reset();
 	}
 
 	/**
-	Initialize the stream for reading with the specified data bytes. Bytes will be read assuming the most significant bit in the byte comes logically
-	earlier in the stream.
+	Restart reading from the stream. Reading will begin anew from the start of the stream.
 	*/
-	FHuffmanBitStreamReader(const TArray<uint8> &SetBytes)
+	void Reset()
 	{
-		Bytes = SetBytes;
-		Reset();
+		BitBuffer = 0;
+		BitBufferBits = 0;
+		BytePos = 0;
 	}
 
 	/**
-	Restart reading from the stream. Reading will begin anew from the start of the stream.
+	Fill the internal bit buffer. Will not check against buffer bounds, so make sure the input buffer is large enough to read all the bits requested +16 bytes.
+	After the call the bit buffer is guaranteed to contain at least MINIMUM_BITS_AFTER_REFILL valid bits.
 	*/
-	void Reset()
+	FORCEINLINE void Refill()
 	{
-		ReadPos = -1;
-		BitsLeftInByte = 0;
+		const uint8* BytesData = Bytes;
+#if USE_UNALIGNED_READ
+		// Branchless buffer refill
+		checkSlow(BytePos + 7 < NumBytes);	// Make sure entire read uint64 is within buffer bounds
+		BitBuffer |= *(const uint64*)(BytesData + BytePos) << BitBufferBits;
+		BytePos += (63 - BitBufferBits) >> 3;
+		BitBufferBits |= 56;
+#else
+		// Read one byte at a time. Doesn't require unaligned reads and is agnostic to endianness.
+		while (BitBufferBits <= 56)
+		{
+			checkSlow(BytePos < NumBytes);
+			uint8 Byte = BytesData[BytePos++];
+			BitBuffer |= uint64(Byte) << BitBufferBits;
+			BitBufferBits += 8;
+		}
+#endif
 	}
 
 	/**
@@ -193,17 +188,7 @@ public:
 	*/
 	FORCEINLINE uint32 Read()
 	{
-		// Advance to the next byte
-		if (BitsLeftInByte == 0)
-		{
-			ReadPos++;
-			BitsLeftInByte = 8;
-			checkf(ReadPos < Bytes.Num(), TEXT("Read past the end of the bitstream"));
-		}
-
-		uint32 Result = FHuffmanBitStreamWriter::ExtractBits(Bytes[ReadPos], BitsLeftInByte - 1, 1);
-		BitsLeftInByte--;
-		return Result;
+		return Read(1);
 	}
 
 	/**
@@ -211,96 +196,62 @@ public:
 	*/
 	FORCEINLINE uint32 Read(uint32 NumBits)
 	{
-		uint32 BitsLeft = NumBits;
-		uint32 Result = 0;
-		while (BitsLeft > 0)
-		{
-			// Advance to the next byte
-			if (BitsLeftInByte == 0)
-			{
-				ReadPos++;
-				BitsLeftInByte = 8;
-				checkf(ReadPos < Bytes.Num(), TEXT("Read past the end of the bitstream"));
-			}
-
-			const uint32 NumBitsToRead = FMath::Min(BitsLeftInByte, BitsLeft);
-			Result = Result << NumBitsToRead;
-			Result |= FHuffmanBitStreamWriter::ExtractBits(Bytes[ReadPos], BitsLeftInByte - NumBitsToRead, NumBitsToRead);
-			BitsLeft -= NumBitsToRead;
-			BitsLeftInByte -= NumBitsToRead;
-		}
-		return Result;
+		Refill();
+		uint32 Value = BitBuffer & ((1ull << NumBits) - 1ull);
+		BitBuffer >>= NumBits;
+		BitBufferBits -= NumBits;
+		return Value;
 	}
 
 	/**
-	Skip the next NumBits in the stream.
-	This "may" be slightly faster than reading the bits and discarding the result.
+	Read the next NumBits from the stream without refilling the bit buffer.
 	*/
-	void Advance(uint32 NumBits)
+	FORCEINLINE uint32 ReadNoRefill(uint32 NumBits)
 	{
-		uint32 BitsLeft = NumBits;
-		while (BitsLeft > 0)
-		{
-			// Advance to the next byte
-			if (BitsLeftInByte == 0)
-			{
-				ReadPos++;
-				BitsLeftInByte = 8;
-				checkf(ReadPos < Bytes.Num(), TEXT("Read past the end of the bitstream"));
-			}
+		uint32 Value = BitBuffer & ((1ull << NumBits) - 1ull);
+		BitBuffer >>= NumBits;
+		BitBufferBits -= NumBits;
+		return Value;
+	}
 
-			const uint32 NumBitsToRead = FMath::Min(BitsLeftInByte, BitsLeft);
-			BitsLeft -= NumBitsToRead;
-			BitsLeftInByte -= NumBitsToRead;
-		}
+	/**
+	Read the next NumBits from the stream without refilling the bit buffer.
+	*/
+	FORCEINLINE void SkipNoRefill(uint32 NumBits)
+	{
+		BitBuffer >>= NumBits;
+		BitBufferBits -= NumBits;
 	}
 
 	/**
 	Return the next NumBits in the stream but do not advance the read position.
-	If this reads past the end zero bits will be returned for the bits beyond the end of the stream.
 	*/
-	FORCEINLINE int32 PeekWithZeros(uint32 NumBits)
+	FORCEINLINE int32 Peek(uint32 NumBits)
 	{
-		uint32 BitsLeft = NumBits;
-		uint32 Result = 0;
-		int32 PeekPos = ReadPos;
-		uint32 PeekBitsLeftInByte = BitsLeftInByte;
-		while (BitsLeft > 0)
-		{
-			// Advance to the next byte
-			if (PeekBitsLeftInByte == 0)
-			{
-				PeekPos++;
-				PeekBitsLeftInByte = 8;
-			}
-
-			const uint32 NumBitsToRead = FMath::Min(PeekBitsLeftInByte, BitsLeft);
-			Result = Result << NumBitsToRead;
-			Result |= FHuffmanBitStreamWriter::ExtractBits((PeekPos >= Bytes.Num()) ? 0 : Bytes[PeekPos], PeekBitsLeftInByte - NumBitsToRead, NumBitsToRead);
-			BitsLeft -= NumBitsToRead;
-			PeekBitsLeftInByte -= NumBitsToRead;
-		}
-		return Result;
+		Refill();
+		return BitBuffer & ((1ull << NumBits) - 1ull);
 	}
 
 	/**
-	Get the total number of bytes in the stream.
+	Return the next NumBits in the stream but do not advance the read position without refilling the bit buffer.
 	*/
-	FORCEINLINE int32 GetNumBytes()
+	FORCEINLINE int32 PeekNoRefill(uint32 NumBits)
 	{
-		return Bytes.Num();
+		return BitBuffer & ((1ull << NumBits) - 1ull);
 	}
 
 	/**
-	Note GetNumBits intentionally does not exist. This class doesn't know exactly how much bits there are
-	it only knows this up to a multiple of 8 the bytes. If data was correctly written using
-	the FHuffmanBitstreamWiter class these extra bits are however guaranteed to be zero.
+	Get the total number of bytes in the stream.
 	*/
-	//uint32 GetNumBits()
-
+	FORCEINLINE int32 GetNumBytes()
+	{
+		return NumBytes;
+	}
 private:
-	TArray<uint8> Bytes;
-	int32 ReadPos;
-	uint32 BitsLeftInByte;
+	const uint8* Bytes;
+	uint64 NumBytes;
+	uint64 BitBuffer;
+	uint32 BitBufferBits;
+	uint32 BytePos;
 };
 
diff --git a/Engine/Plugins/Experimental/GeometryCache/Source/GeometryCache/Private/HuffmanTable.cpp b/Engine/Plugins/Experimental/GeometryCache/Source/GeometryCache/Private/HuffmanTable.cpp
--- a/Engine/Plugins/Experimental/GeometryCache/Source/GeometryCache/Private/HuffmanTable.cpp
+++ b/Engine/Plugins/Experimental/GeometryCache/Source/GeometryCache/Private/HuffmanTable.cpp
@@ -10,469 +10,249 @@ To read the algorithmic docs you need to understand the following terms
 - Codeword: The Huffman code. A variable length bit pattern.
 - Some symbols will get a codeword assigned others which never occurred will not get a codeword assigned
 - Symbols which occur more often preferably get a shorter code word
-
-We often reference the Huffman algorithm as used by JPEG standard in the comments. This code uses an almost identical
-algorithm. The main differences with the JPEG standard are:
-- Arbitrary code lengths and symbols (JPEG is limited to 16 bits length, with 256 symbols)
-- No special codes (all 1 bits) which JPEG uses as bit stream sync points.
-
 */
 
-#define NO_LOOKAHEAD 0
-
-/**
-	Uninitialized Huffman table, call Initialize() first
-*/
-FHuffmanTable::FHuffmanTable()
-{	
-}
 
 /**
-	Initialize a Huffman table for encoding values in the range [0..NumSymbols[
-	The table will initially be in prepass mode.
+	Reverses the bits of an unsigned integer of a given number of bits.
 */
-void FHuffmanTable::Initialize(int32 InNumSymbols)
+static uint32 ReverseBits(uint32 Value, uint32 NumBits)
 {
-	NumSymbols = InNumSymbols;
-	SymbolFrequencies.Init(0, NumSymbols + 1);
-	CodewordSymbols.Init(0, NumSymbols + 1);
-	// There may be less codewords than symbols but never more so this is a safe bet for the size of this array.
-	CodeWords.Init(FCodeWord(), NumSymbols);
-	ResetFrequencies();
-	bIsPrepass = true;
+	Value = ((Value & 0x55555555u) << 1u) | ((Value & 0xAAAAAAAAu) >> 1u);
+	Value = ((Value & 0x33333333u) << 2u) | ((Value & 0xCCCCCCCCu) >> 2u);
+	Value = ((Value & 0x0F0F0F0Fu) << 4u) | ((Value & 0xF0F0F0F0u) >> 4u);
+	Value = ((Value & 0x00FF00FFu) << 8u) | ((Value & 0xFF00FF00u) >> 8u);
+	Value = (Value << 16u) | (Value >> 16u);
+	return Value >> (32u - NumBits);
 }
 
 /**
-	Initialize a Huffman table based on the serialized table.
-	The table will initially be in encode/decode mode.
+	Calculate optimal length-limited Huffman code lengths using the package-merge algorithm.
+	Reference: https://www.ics.uci.edu/~dan/pubs/LenLimHuff.pdf
 */
-void FHuffmanTable::Initialize(FHuffmanBitStreamReader &Stream)
+static void GenerateLengthLimitedHuffmanLengths(uint8* OutLengths, uint32* Frequencies, int32 NumSymbols, int32 MaxCodeLength)
 {
-	NumSymbols = Stream.Read(32);
+	check(NumSymbols <= (1 << MaxCodeLength));
+	check(NumSymbols <= HUFFMAN_MAX_CODES);
+	check(MaxCodeLength <= HUFFMAN_MAX_CODE_LENGTH);
 
-	// We only serialize enough bits depending on the number of symbols available.
-	int32 LastSymbol = NumSymbols - 1;
-	uint64 HighestOne = FPlatformMath::FloorLog2(LastSymbol); // FloorLog2 is basically returning the highest set bit (zero based index)
-	int32 SymbolBitsToSerialize = HighestOne + 1;
+	FMemory::Memzero(OutLengths, NumSymbols);
 
-	// Initialize tables
-	SymbolFrequencies.Empty();
-
-	CodewordSymbols.Empty(NumSymbols + 1);
-	CodewordSymbols.SetNumZeroed(NumSymbols + 1);
+	struct SortEntry
+	{
+		uint32 Freq;
+		int32 Symbol;
+	};
 
-	CodeWords.Empty(NumSymbols);
-	CodeWords.AddDefaulted(NumSymbols);
+	// Gather symbols with frequency > 0
+	TArray<SortEntry> SortedSymbols;
+	SortedSymbols.Reserve(NumSymbols);
 
-	// Load NumCodeWordsOfLength table
-	NumCodeWordsOfLength[0] = 0;
-	for (int32 Length = 1; Length < MaxCodeLength + 1; Length++)
+	int32 NonZeroIndex = 0;
+	for (int32 Symbol = 0; Symbol < NumSymbols; Symbol++)
 	{
-		NumCodeWordsOfLength[Length] = Stream.Read(32);
+		uint32 Freq = Frequencies[Symbol];
+		if (Freq > 0)
+		{
+			NonZeroIndex = Symbol;
+			SortedSymbols.Emplace(SortEntry{ Freq, Symbol });
+		}
 	}
 
-	// Calculate this instead of storing it in the stream
-	CountTotalCodewords();
+	// Sort symbols by frequency
+	SortedSymbols.Sort([](const SortEntry& A, const SortEntry& B) { return A.Freq < B.Freq; });
 
-	// LoadCodewordSymbols table
-	for (int32 CodeWordId = 0; CodeWordId < TotalCodeWords; CodeWordId++)
+	// If there is only one symbol give it length 1 to avoid having to handle 0-length codes
+	int32 NumUsedSymbols = SortedSymbols.Num();
+	if (NumUsedSymbols == 1)
 	{
-		CodewordSymbols[CodeWordId] = 0;
-		CodewordSymbols[CodeWordId] = Stream.Read(SymbolBitsToSerialize);
+		OutLengths[NonZeroIndex] = 1;
+		return;
 	}
 
-	// Generate derived tables
-	GenerateCodewords();
-	bIsPrepass = false;
-}
+	struct Node
+	{
+		int32	Symbol;
+		uint32	Freq;
+		Node*	Children[2];
 
-/**
-	Serialize the Huffman table to the bitstream, deserialize by calling Initialise(Reader)
-*/
-void FHuffmanTable::Serialize(FHuffmanBitStreamWriter& Stream)
-{
-	Stream.Write(NumSymbols, 32);
+		// Traverse tree and every time a leaf node is visited the associated symbol code length is incremented by 1.
+		static void GenerateLengthLimitedHuffmanLengthsInner(Node* Node, uint8* OutLengths)
+		{
+			if (Node->Symbol == INT32_MIN)
+			{
+				GenerateLengthLimitedHuffmanLengthsInner(Node->Children[0], OutLengths);
+				GenerateLengthLimitedHuffmanLengthsInner(Node->Children[1], OutLengths);
+			}
+			else
+			{
+				OutLengths[Node->Symbol]++;
+			}
+		}
+	};
 
-	// We only serialize enough bits depending on the number of symbols available.
-	int32 LastSymbol = NumSymbols - 1;
-	uint64 HighestOne = FPlatformMath::FloorLog2(LastSymbol); // FloorLog2 is basically returning the highest set bit (zero based index)
-	int32 SymbolBitsToSerialize = HighestOne + 1;
+	// Loop over code lengths in ascending order. At every length merge pairs of nodes from the previous level into the current whenever it is cheaper.
+	TArray<Node> Nodes;
+	Nodes.InsertZeroed(0, NumUsedSymbols * MaxCodeLength * 2);
 
-	// Save NumCodeWordsOfLength table
-	for (int32 Length = 1; Length < MaxCodeLength + 1; Length++)
+	int32 CurrentNodeIndex = 0;
+	int32 PrevNodeStartIndex = 0;
+	int32 NumPrevNodes = 0;
+	for (int32 CodeLength = 1; CodeLength <= MaxCodeLength; CodeLength++)
 	{
-		Stream.Write(NumCodeWordsOfLength[Length], 32);
+		int32 NumCurrentNodes = NumUsedSymbols;
+		int32 PrevNodeIndex = PrevNodeStartIndex;
+		PrevNodeStartIndex = CurrentNodeIndex;
+
+		int32 NumWrittenNodes = 0;
+		while (NumPrevNodes >= 2 || NumCurrentNodes > 0)
+		{
+			// Merge pair of nodes from previous level into current level when they have a combined Freq smaller than the next node in the current level
+			if (NumCurrentNodes > 0 && (NumPrevNodes < 2 || SortedSymbols[NumUsedSymbols - NumCurrentNodes].Freq <= Nodes[PrevNodeIndex + 0].Freq + Nodes[PrevNodeIndex + 1].Freq))
+			{
+				// Select node from current level
+				const SortEntry& Entry = SortedSymbols[NumUsedSymbols - NumCurrentNodes];
+				Node& Node = Nodes[CurrentNodeIndex++];
+				Node.Freq = Entry.Freq;
+				Node.Symbol = Entry.Symbol;
+				Node.Children[0] = nullptr;
+				Node.Children[1] = nullptr;
+				NumCurrentNodes--;
+			}
+			else
+			{
+				// Merge two nodes from the previous level
+				Node& Node = Nodes[CurrentNodeIndex++];
+				Node.Freq = Nodes[PrevNodeIndex + 0].Freq + Nodes[PrevNodeIndex + 1].Freq;
+				Node.Symbol = INT32_MIN;
+				Node.Children[0] = &Nodes[PrevNodeIndex + 0];
+				Node.Children[1] = &Nodes[PrevNodeIndex + 1];
+				NumPrevNodes -= 2;
+				PrevNodeIndex += 2;
+			}
+			NumWrittenNodes++;
+		}
+		NumPrevNodes = NumWrittenNodes;
 	}
 
-	// Save CodewordSymbols table
-	for (int32 CodeWordId = 0; CodeWordId < TotalCodeWords; CodeWordId++)
+	// Traverse final level to update symbol lengths
+	check(CurrentNodeIndex <= NumUsedSymbols * MaxCodeLength * 2);
+	int32 NumActive = 2 * NumUsedSymbols - 2;
+	for (int32 i = 0; i < NumActive; i++)
 	{
-		Stream.Write(CodewordSymbols[CodeWordId], SymbolBitsToSerialize);
+		Node::GenerateLengthLimitedHuffmanLengthsInner(&Nodes[PrevNodeStartIndex + i], OutLengths);
 	}
-
 }
 
 /**
-Toggle between prepass mode and encoding/decoding mode.
-Note: Toggling will update internal data structures and may have take more time than setting a simple variable.
+	Generate canonical Huffman codes from symbol lengths.
+	Codes are reversed, so they are actually postfix free instead of prefix free.
+	Postfix codes are needed when bits are parsed from the bottom.
 */
-void FHuffmanTable::SetPrepass(bool bInIsPrepass)
+static void GenerateHuffmanCodes(uint16* OutCodes, const uint8* CodeLengths, int32 NumSymbols, int32 MaxCodeLength)
 {
-	// If we enabled prepass reset the histogram
-	if (bIsPrepass == false && bInIsPrepass == true)
+	int32 SymbolLengthHistogram[HUFFMAN_MAX_CODE_LENGTH + 1] = {};			// Histogram of code lengths
+	int16 SymbolLists[(HUFFMAN_MAX_CODE_LENGTH + 1) * HUFFMAN_MAX_CODES];	// List of symbols for every code length
+
+	// Fill histogram and generate symbol lists for each code length
+	for (int32 Symbol = 0; Symbol < NumSymbols; Symbol++)
 	{
-		SymbolFrequencies.Init(0, NumSymbols + 1);
-		ResetFrequencies();
+		int32 CodeLength = CodeLengths[Symbol];
+		check(CodeLength <= HUFFMAN_MAX_CODE_LENGTH);
+		SymbolLists[CodeLength * HUFFMAN_MAX_CODES + SymbolLengthHistogram[CodeLength]++] = Symbol;
 	}
-	// If we disabled prepass calculate the tables
-	if (bIsPrepass == true && bInIsPrepass == false)
+
+	// Assign code words to symbols
+	int32 NextCodeWord = 0;
+	for (int32 CodeLength = 1; CodeLength <= MaxCodeLength; CodeLength++)
 	{
-		GenerateOptimalTable();
-		// We don't need the frequency table anymore once we generated the tables
-		SymbolFrequencies.Empty(0);
+		// Loop over symbols in ascending order of current code length and assign consecutive codes to them starting from 2x the last code of the previous code length.
+		int32 NumSymbolsOfLength = SymbolLengthHistogram[CodeLength];
+		for (int32 i = 0; i < NumSymbolsOfLength; i++)
+		{
+			int32 Symbol = SymbolLists[CodeLength * HUFFMAN_MAX_CODES + i];
+			check(CodeLengths[Symbol] == CodeLength);
+			OutCodes[Symbol] = ReverseBits(NextCodeWord++, CodeLength);	// Reverse bits to turn prefix codes into postfix codes
+		}
+		NextCodeWord <<= 1;
 	}
-	bIsPrepass = bInIsPrepass;
 }
 
 /**
-	Reset the frequency of all symbols to zero.
+	Initialize a Huffman table for encoding values in the range [0..NumSymbols[
+	The table will initially be in prepass mode.
 */
-void FHuffmanTable::ResetFrequencies()
+void FHuffmanEncodeTable::Initialize(int32 InNumSymbols)
 {
-	for (int32 i = 0; i < SymbolFrequencies.Num(); i++)
-	{
-		SymbolFrequencies[i] = 0;
-	}
+	NumSymbols = InNumSymbols;
+	SymbolFrequencies.Init(0, NumSymbols);
+	CodewordSymbols.Init(0, NumSymbols);
+	// There may be less codewords than symbols but never more so this is a safe bet for the size of this array.
+	CodeWords.Init(FCodeWord(), NumSymbols);
+	ResetFrequencies();
+	bIsPrepass = true;
 }
 
 /**
-	Bump the frequency of all symbols by one.
+	Serialize the Huffman table to the bitstream. Deserialize by calling FHuffmanDecodeTable.Initialise()
 */
-void FHuffmanTable::MarkAllSymbols()
+void FHuffmanEncodeTable::Serialize(FHuffmanBitStreamWriter& Stream)
 {
-	check(bIsPrepass); // This only makes sense when we are doing a prepass
+	Stream.Write(NumSymbols, HUFFMAN_SYMBOL_COUNT_BITS);
 
-	// We simply make all frequencies at least one here. But we also bump the count of the
-	// actually occurring frequencies to give them a small advantage over the ones we never
-	// actually saw occurring in the sample stream.
-	for (int32 i = 0; i < SymbolFrequencies.Num(); i++)
+	// Write symbol lengths. We don't need to write the actual codes as they can be canonically reconstructed from the lengths using GenerateHuffmanCodes
+	for (int32 Symbol = 0; Symbol < NumSymbols; Symbol++)
 	{
-		SymbolFrequencies[i] += 1;
+		Stream.Write(CodeWords[Symbol].Length, HUFFMAN_MAX_CODE_LENGTH_BITS);
 	}
 }
 
 /**
-	Generate the NumCodeWordsOfLength and CodewordSymbols tables based on the SymbolFrequencies table.
+Toggle between prepass mode and encoding/decoding mode.
+Note: Toggling will update internal data structures and may have take more time than setting a simple variable.
 */
-void FHuffmanTable::GenerateOptimalTable()
+void FHuffmanEncodeTable::SetPrepass(bool bInIsPrepass)
 {
-	// We first assign temporary code lengths to symbols. This value is the maximum length
-	// of these temporary symbols. If we generate codes longer than this length we cannot continue the generation.
-	// These temporary code lengths will then be shortened to the MaxCodeLength set by the user. At a cost
-	// of optimality in the Huffman table.
-	const int32 MaxTemporaryCodeLength = 32;
-	int32 CodesWithLength[MaxTemporaryCodeLength + 1]; // Number of codes with length i
-	
-	TArray<int32> CodeLengths;
-	CodeLengths.Init(0, NumSymbols + 1);
-	TArray<int32> Others;
-	Others.Init(0, NumSymbols + 1);
-
-	/* This algorithm is explained in section K.2 of the JPEG standard */
-
-	//
-	// Initialize tables
-	//
-	memset(CodesWithLength, 0, sizeof(int) * (MaxTemporaryCodeLength +1));
-	for (int32 Symbol = 0; Symbol < (NumSymbols + 1); Symbol++)
-	{
-		Others[Symbol] = -1;
-	}
-
-	SymbolFrequencies[NumSymbols] = 1;		// make sure the last, sentinel, symbol has a nonzero count
-
-	//
-	// Here we fill the others and codesize arrays
-	// This looks a bit weird it doesn't actually explicitly build the trees it just generates some 
-	// bookkeeping info on the trees like frequencies and number of bits.
-	//
-	for (;;)
-	{
-		// Find the largest symbol with the lowest frequency greater than zero.
-		// i.e. if two codewords have equal frequency select the largest of the two.
-		int32 Largest1 = -1;
-		uint64 lowestFreq = 0xffffffffffffffff/*ui64*/;
-		for (int32 Symbol=0; Symbol <= NumSymbols; Symbol++)
-		{
-			if (SymbolFrequencies[Symbol] && SymbolFrequencies[Symbol] <= lowestFreq)
-			{
-				lowestFreq = SymbolFrequencies[Symbol];
-				Largest1 = Symbol;
-			}
-		}
-
-		// Find the second largest symbol with the lowest frequency greater than zero.
-		// again, if two codewords have equal frequency select the largest of the two.
-		int32 Largest2 = -1;
-		lowestFreq = 0xffffffffffffffff/*ui64*/;
-		for (int32 Symbol = 0; Symbol <= NumSymbols; Symbol++)
-		{
-			if (SymbolFrequencies[Symbol] && SymbolFrequencies[Symbol] <= lowestFreq && Symbol != Largest1)
-			{
-				lowestFreq = SymbolFrequencies[Symbol];
-				Largest2 = Symbol;
-			}
-		}
-
-		// Done if we've merged everything down into one root node
-		if (Largest2 < 0)
-			break;
-
-		// Largest1 and Largest2 are now the two codewords we have to put together as childs of a new node...
-		// we store this "new" node in index Largest1
-		SymbolFrequencies[Largest1] += SymbolFrequencies[Largest2]; // Combined frequencies of the two nodes
-		SymbolFrequencies[Largest2] = 0; // Other node gets removed
-
-		// We need an additional bit to make the difference between Largest1 sub tree and Largest2 sub tree
-		// so update the books accordingly for every node in those subtrees
-
-		// Increment the code size with one bit of everything in Largest1's tree branch
-		CodeLengths[Largest1]++;
-		while (Others[Largest1] >= 0)
-		{
-			Largest1 = Others[Largest1];
-			CodeLengths[Largest1]++;
-		}
-
-		// Attach Largest2 in the newly merged tree
-		Others[Largest1] = Largest2;
-
-		// Increment the code size of everything in Largest2's tree branch
-		CodeLengths[Largest2]++;
-		while (Others[Largest2] >= 0)
-		{
-			Largest2 = Others[Largest2];
-			CodeLengths[Largest2]++;
-		}
-	}
-
-	// Now count the number of symbols of each code length and store in the bits array
-	// we include the sentinel element here
-	for (int32 Symbol = 0; Symbol < NumSymbols+1; Symbol++)
-	{
-		if (CodeLengths[Symbol])
-		{
-			// We have so many symbols that the code length is already longer than our internal maximum
-			check(CodeLengths[Symbol] < MaxTemporaryCodeLength + 1);
-			CodesWithLength[CodeLengths[Symbol]]++;
-		}
-	}
-
-	/* We don't allow symbols with code lengths over MaxCodeLength bits, so if the pure
-	Huffman procedure assigned any such lengths, we must adjust the coding.
-	So we go against the optimal tree specified by frequencies in order to make
-	it more balanced assigning shorter codes to two subtrees.
-	
-	Note this code only again only does the bookkeeping of how many symbols assigned
-	to what lengths at first it seems magic but once you realize it's just doing
-	bookkeeping it's all rather straightforward.
-	
-	This rather magic code is explained as follows in the spec (sic.):
-	Since symbols are paired for the longest Huffman code, the symbols are
-	removed from this length category two at a time.  The prefix for the pair
-	(which is one bit shorter) is allocated to one of the pair; then,
-	skipping the BITS entry for that prefix length, a code word from the next
-	shortest nonzero BITS entry is converted into a prefix for two code words
-	one bit longer.
-	*/
-
-	int32 CodeLength;
-	for (CodeLength = MaxTemporaryCodeLength; CodeLength > MaxCodeLength; CodeLength--)
+	// If we enabled prepass reset the histogram
+	if (bIsPrepass == false && bInIsPrepass == true)
 	{
-		while (CodesWithLength[CodeLength] > 0)
-		{
-			// find length of new prefix to be used
-			// we take one of the codewords and assign it a different length below
-			// so it can't be 0 which is why we need the loop here
-			int32 PrefixLen = CodeLength - 2;
-			while (CodesWithLength[PrefixLen] == 0)
-				PrefixLen--;
-			
-			// So now we take one of the symbols assigned a codeword of 'prefixLen' bits
-			// and assign that to a codeword of codeLength-1 bits. We then
-			// use the freed up prefixLen-bit codeword as a prefix for two symbols of prefixLen+1 bits.
-
-			CodesWithLength[CodeLength] -= 2;		// remove two symbols from this length
-			CodesWithLength[CodeLength-1] += 1;		// the removal above freed one new symbol with length length codeLength-1
-			CodesWithLength[PrefixLen+1] += 2;		// two new symbols with length prefixLen+1
-			CodesWithLength[PrefixLen] -= 1;		// One of the codewords of this length is now a prefix (for the two j+1 ones in the previous line) and no longer assigned to any symbol
-		}
+		SymbolFrequencies.Init(0, NumSymbols);
+		ResetFrequencies();
 	}
-
-	// Remove the count for the pseudo-symbol 'NumSymbols' from the largest code length still in use
-	while (CodesWithLength[CodeLength] == 0)
-		CodeLength--;
-	CodesWithLength[CodeLength]--;
-
-	// Save final symbol counts
-	memcpy(NumCodeWordsOfLength, CodesWithLength, sizeof(NumCodeWordsOfLength));
-
-	// The list of symbols is now sorted according to code size.
-	// JPEG spec seems to think this works without considering the new lengths we assigned above ?!?
-	// Probably because these new lengths didn't change the relative code word order only the code word lengths....
-
-	// Create a reverse lookup table converting a code word to a symbol
-	// Just loop over it from shortest to longest length
-	for (int32 Length=1, CodeWordId=0; Length<=MaxTemporaryCodeLength; Length++)
+	// If we disabled prepass calculate the tables
+	if (bIsPrepass == true && bInIsPrepass == false)
 	{
-		// Just do a brute force search for codewords with this length
-		for (int32 Symbol = 0; Symbol<NumSymbols; Symbol++)
+		uint8 SymbolLengths[HUFFMAN_MAX_CODES];
+		uint16 SymbolCodes[HUFFMAN_MAX_CODES];
+		GenerateLengthLimitedHuffmanLengths(SymbolLengths, SymbolFrequencies.GetData(), NumSymbols, HUFFMAN_MAX_CODE_LENGTH);
+		GenerateHuffmanCodes(SymbolCodes, SymbolLengths, NumSymbols, HUFFMAN_MAX_CODE_LENGTH);
+
+		CodeWords.Empty();
+		CodeWords.SetNumUninitialized(NumSymbols);
+		for (int32 i = 0; i < NumSymbols; i++)
 		{
-			if (CodeLengths[Symbol] == Length)
-			{
-				CodewordSymbols[CodeWordId] = Symbol;
-				CodeWordId++;
-			}
+			CodeWords[i].Bits = SymbolCodes[i];
+			CodeWords[i].Length = SymbolLengths[i];
 		}
-	}
-
-	// Update our derived values
-	CountTotalCodewords();
-	GenerateCodewords();
-}
 
-void FHuffmanTable::CountTotalCodewords()
-{
-	// As some symbols may have frequency 0 they are not assigned a codeword
-	// and thus totalCodeWords may be smaller than the number of symbols 
-	TotalCodeWords = 0;
-	for (int32 Length = 1; Length<MaxCodeLength; Length++)
-	{
-		TotalCodeWords += NumCodeWordsOfLength[Length];
+		// We don't need the frequency table anymore once we generated the tables
+		SymbolFrequencies.Empty(0);
 	}
+	bIsPrepass = bInIsPrepass;
 }
 
 /**
-	Generate the codeWords, DecodeTable, and Lookahead tables based on the
-	NumCodeWordsOfLength and CodewordSymbols tables.
+	Reset the frequency of all symbols to zero.
 */
-void FHuffmanTable::GenerateCodewords()
+void FHuffmanEncodeTable::ResetFrequencies()
 {
-	// There are always less or equal codewords compared to code values so we can safely preallocate these
-	// to NumValues, we add a sentinel element as well.
-	TArray<int8> CodeWordLengths; //size of the i'th code word
-	CodeWordLengths.Init(0, NumSymbols + 1);
-	TArray<uint32> CodeWordBits; // Bit pattern of the i'th code word
-	CodeWordBits.Init(0, NumSymbols + 1);
-
-	// Figure C.1 in JPEG: create a table of code length for each code word. Note that this is in increasing code-length order.
-	int32 CodeWordId = 0;
-	for (int32 Length = 1; Length <= MaxCodeLength; Length++)
-	{
-		for (int32 Counter = 1; Counter <= (int)NumCodeWordsOfLength[Length]; Counter++)
-		{
-			CodeWordLengths[CodeWordId++] = (char)Length;
-		}
-	}
-	CodeWordLengths[CodeWordId] = 0;
-	int32 LastCodeWordId = CodeWordId;
-
-	// Figure C.2 in JPEG: generate the codes themselves in increasing order
-	int32 Code = 0;
-	int32 CurrentLength = CodeWordLengths[0];
-	CodeWordId = 0;
-	while (CodeWordLengths[CodeWordId])
-	{
-		// Loop over all codewords of length 'CurrentLength' and assign codes
-		while (((int)CodeWordLengths[CodeWordId]) == CurrentLength)
-		{
-			CodeWordBits[CodeWordId++] = Code;
-			Code++;
-		}
-		// Increase the code length
-		Code <<= 1;
-		CurrentLength++;
-	}
-
-	// Figure C.3 in JPEG: generate encoding tables
-	// These are code and size indexed by symbol value
-
-	// As not all symbols may have a codeword assigned we set the length initially to 0 for those
-	for (int32 SymbolIndex = 0; SymbolIndex < NumSymbols; SymbolIndex++)
-	{
-		CodeWords[SymbolIndex].Bits = 0;
-		CodeWords[SymbolIndex].Length = 0;
-	}
-
-	// Any assign the codeword for any symbols which got one assigned
-	for (CodeWordId = 0; CodeWordId < LastCodeWordId; CodeWordId++)
-	{
-		CodeWords[CodewordSymbols[CodeWordId]].Bits = CodeWordBits[CodeWordId];
-		CodeWords[CodewordSymbols[CodeWordId]].Length = CodeWordLengths[CodeWordId];
-	}
-
-	//
-	// Generate the decoding tables
-	//
-
-	// Figure F.15 in JPEG: generate decoding tables for bit-sequential decoding
-
-	CodeWordId = 0;
-	DecodeTable[0].MaxCode = DecodeTable[0].MinCode = DecodeTable[0].ValPtr = 0;
-
-	for (int32 Length = 1; Length <= MaxCodeLength; Length++)
-	{
-		if (NumCodeWordsOfLength[Length])
-		{
-			DecodeTable[Length].ValPtr = CodeWordId; // index of 1st symbol of code length l
-			DecodeTable[Length].MinCode = CodeWordBits[CodeWordId]; // minimum code of length l
-			CodeWordId += NumCodeWordsOfLength[Length];
-			DecodeTable[Length].MaxCode = CodeWordBits[CodeWordId - 1]; // maximum code of length l
-		}
-		else
-		{		
-			DecodeTable[Length].MaxCode = -1;	// -1 if no codes of this length
-			// Just zero the rest these values shouldn't ever matter
-			DecodeTable[Length].MinCode = 0;
-			DecodeTable[Length].ValPtr = 0;
-		}
-	}
-
-	DecodeTable[MaxCodeLength+1].MaxCode = 0xFFFFFFL; // ensures the decoder always terminates
-
-	// This is an optimization not forced in the jpeg standard but commonly used in decoders.
-	// We make a quick lookup table which allows us to look at the next byte (8 bits) coming up in the
-	// bit stream and quickly determine what the corresponding symbol is (obviously only for codewords shorter than 8 bits)
-
-	// Initialize to 0 which means we can't recognize a symbol in the lookahead so it's coding a prefix
-	// for a symbol longer than LookaheadBits
-	for (int32 i = 0; i < LookaheadEntries; i++)
-	{
-		Lookahead[i].Length = 0;
-		Lookahead[i].Symbol = 0;
-	}
-
-	CodeWordId = 0;
-	for (int32 Length = 1; Length <= LookaheadBits; Length++)
+	for (int32 i = 0; i < SymbolFrequencies.Num(); i++)
 	{
-		for (int32 Counter = 1; Counter <= (int)NumCodeWordsOfLength[Length]; Counter++, CodeWordId++)
-		{
-			// Left justified codeword bit sequence. Rightmost bits are zero
-			int32 Bits = CodeWordBits[CodeWordId] << (LookaheadBits - Length);
-
-			// Initialize table entries at all possible bit sequences that may follow the current codeword
-			// They all refer to the same codeword and length
-			for (int32 Entry = 1 << (LookaheadBits - Length); Entry > 0; Entry--)
-			{
-				Lookahead[Bits].Length = Length;
-				Lookahead[Bits].Symbol = CodewordSymbols[CodeWordId];
-				Bits++;
-			}
-		}
+		SymbolFrequencies[i] = 0;
 	}
 }
 
-void FHuffmanTable::Encode(FHuffmanBitStreamWriter& Stream, int32 Symbol)
+void FHuffmanEncodeTable::Encode(FHuffmanBitStreamWriter& Stream, int32 Symbol)
 {
 	// Check if symbol is in the valid range
 	check(Symbol >= 0);
@@ -493,53 +273,47 @@ void FHuffmanTable::Encode(FHuffmanBitStreamWriter& Stream, int32 Symbol)
 	}
 }
 
-int32 FHuffmanTable::LongDecode(FHuffmanBitStreamReader& Stream, int32 MinLength)
+/**
+	Initialize a Huffman table based on the serialized table.
+	The table will initially be in encode/decode mode.
+*/
+void FHuffmanDecodeTable::Initialize(FHuffmanBitStreamReader &Stream)
 {
-	int32 Code = 0;
+	int32 NumSymbols = Stream.Read(HUFFMAN_SYMBOL_COUNT_BITS);
+	check(NumSymbols <= HUFFMAN_MAX_CODES);
 
-	//We need to at least decode MinLength bits possibly more
-	check(MinLength >= 1);
-	Code = Stream.Read(MinLength);
+	uint8 SymbolLengths[HUFFMAN_MAX_CODES];
 
-	// Read the rest one bit at a time (Figure F.16 in the JPEG spec)
-	int32 Length = MinLength;
-	while (Length <= MaxCodeLength && Code > DecodeTable[Length].MaxCode)
+	// Read symbol code lengths
+	for (int32 Symbol = 0; Symbol < NumSymbols; Symbol++)
 	{
-		Code <<= 1;
-		Code |= Stream.Read(1);
-		Length++;
+		SymbolLengths[Symbol] = Stream.Read(HUFFMAN_MAX_CODE_LENGTH_BITS);
 	}
 
-	// Guard against bit stream errors
-	if (Length > MaxCodeLength)
-	{
-		return 0;
-	}
-
-	const int32 OffsetInCodesOfLength = Code - DecodeTable[Length].MinCode;
-	return CodewordSymbols[DecodeTable[Length].ValPtr + OffsetInCodesOfLength];
-}
-
-int32 FHuffmanTable::Decode(FHuffmanBitStreamReader &Stream)
-{
-	check(bIsPrepass == false);
-
-#if NO_LOOKAHEAD
-	return LongDecode(Stream, 1);
-#else	
-	//	Peek the first valid byte	
-	const int32 Bits = Stream.PeekWithZeros(LookaheadBits); // Should return zeros if past end
+	// Generate canonical huffman codes from code lengths
+	uint16 SymbolCodes[HUFFMAN_MAX_CODES];
+	GenerateHuffmanCodes(SymbolCodes, SymbolLengths, NumSymbols, HUFFMAN_MAX_CODE_LENGTH);
 
-	// Found in lookahead?
-	const int32 Length = Lookahead[Bits].Length;
-	if (Length)
-	{
-		Stream.Read(Length);
-		return Lookahead[Bits].Symbol;
-	}
-	else
+	// Generate decode lookup table
+	uint32 MaxCode = 1u << HUFFMAN_MAX_CODE_LENGTH;
+	for(int32 Symbol = 0; Symbol < NumSymbols; Symbol++)
 	{
-		return LongDecode(Stream, LookaheadBits + 1);
+		uint32 Length = SymbolLengths[Symbol];
+		if(Length > 0)
+		{
+			uint32 Code = SymbolCodes[Symbol];
+			uint32 Step = 1u << Length;
+			
+			FTableEntry Entry;
+			Entry.Symbol = Symbol;
+			Entry.Length = Length;
+
+			// Put entry in at every position where the Length last bits are Code.
+			do
+			{
+				TableEntries[Code] = Entry;
+				Code += Step;
+			} while (Code < MaxCode);
+		}
 	}
-#endif
 }
diff --git a/Engine/Plugins/Experimental/GeometryCache/Source/GeometryCache/Private/HuffmanTable.h b/Engine/Plugins/Experimental/GeometryCache/Source/GeometryCache/Private/HuffmanTable.h
--- a/Engine/Plugins/Experimental/GeometryCache/Source/GeometryCache/Private/HuffmanTable.h
+++ b/Engine/Plugins/Experimental/GeometryCache/Source/GeometryCache/Private/HuffmanTable.h
@@ -6,10 +6,10 @@
 #include "HuffmanBitStream.h"
 
 /**
-	An encapsulated Huffman coding/decoding class. To use this class roughly follow the following sequence of calls
+	Encapsulated Huffman coding/decoding classes. To use these classes roughly follow the following sequence of calls
 
-	//Create a table
-	HuffmanTable Tab(MaxValue);
+	// Create an encoding table
+	FHuffmanEncodeTable Tab(MaxValue);
 
 	// First pass, gather statistics about the data to code
 	Tab.SetPrepass(true)
@@ -20,10 +20,10 @@
 	for each value to code
 		Tab.Encode(value)
 
-	The generated bits are then saved to disc alongside the generated Huffman table.
+	The generated bits are then saved to disk alongside the generated Huffman table.
 	To decode the following is then done
 
-	HuffmanTable Tab(SerializedTable)
+	FHuffmanDecodeTable Tab(SerializedTable)
 	for each value to decode
 		value = Tab.Decode()
 
@@ -33,62 +33,16 @@
 	- Codeword: The Huffman code. A variable length bit pattern.
 	- Some symbols will get a codeword assigned others which never occurred will not get a codeword assigned
 	- Symbols which occur more often preferably get a shorter code word
-
 */
-class FHuffmanTable
-{
-	// This has a sentinel element at index NumValues!! And thus contains NumValues+1 elements
-	TArray<uint32> SymbolFrequencies;
-	TArray<int32> CodewordSymbols;
-	static const int32 MaxCodeLength = 24;
 
-	struct FCodeWord
-	{
-		int32 Bits;
-		int32 Length;
-	};
-
-	struct FDecodeTableEntry
-	{
-		int32 MinCode;	// contains the smallest code value for a given length I
-		int32 MaxCode;	// contains the largest code value for a given length I
-		int32 ValPtr;	//contains the index to the start of the list of values in HUFFVAL which are decoded by code words of length I.
-	};
-
-	struct FLookaheadEntry
-	{
-		int32 Length; //Length of the codeword we recognized in the lookahead
-		int32 Symbol; //Symbol it corresponds to
-	};
-
-	TArray<FCodeWord> CodeWords; // For every symbol the corresponding huffcode bits
-	int32 NumCodeWordsOfLength[MaxCodeLength+1]; //Number of codewords with a given length, element 0 is unused
-	int32 NumSymbols;
-
-	void GenerateOptimalTable();
-	void CountTotalCodewords();
-	void GenerateCodewords();
-	
-	FDecodeTableEntry DecodeTable[MaxCodeLength+2];
-
-	static const int32 LookaheadBits = 8;
-	static const int32 LookaheadEntries = 1 << LookaheadBits;
-
-	FLookaheadEntry Lookahead[LookaheadEntries];
-
-	int32 TotalCodeWords;
-
-	int32 LongDecode(FHuffmanBitStreamReader& Stream, int32 MinLength);
-	void ResetFrequencies();
-	bool bIsPrepass;
+#define HUFFMAN_MAX_CODES 256
+#define HUFFMAN_MAX_CODE_LENGTH 11
+#define HUFFMAN_MAX_CODE_LENGTH_BITS 4	// Number of bits needed to decode a value in the range [0..HUFFMAN_MAX_CODE_LENGTH]
+#define HUFFMAN_SYMBOL_COUNT_BITS 9		// Number of bits needed to send the number of symbols
 
+class FHuffmanEncodeTable
+{
 public:
-
-	/**
-		Uninitialized Huffman table, call Initialize() first
-	*/
-	FHuffmanTable();
-
 	/**
 		Initialize a Huffman table for encoding values in the range [0..NumSymbols[
 		The table will initially be in prepass mode.
@@ -96,18 +50,12 @@ public:
 	void Initialize(int32 NumSymbols);
 
 	/**
-		Initialize a Huffman table based on the serialized table.
-		The table will initially be in encode/decode mode.
-	*/
-	void Initialize(FHuffmanBitStreamReader& Stream);
-
-	/**
-		Serialize the Huffman table to the bitstream, deserialize by calling Initialise(Reader)
+		Serialize the Huffman table to the bitstream, deserialize using FHuffmanDecodeTable::Initialize
 	*/
-	void Serialize(FHuffmanBitStreamWriter& Stream);	
+	void Serialize(FHuffmanBitStreamWriter& Stream);
 
 	/**
-		Toggle between prepass mode and encoding/decoding mode.
+		Toggle between prepass mode and encoding mode.
 		Note: Toggling will update internal data structures and may have take more time than setting a simple variable.
 	*/
 	void SetPrepass(bool bInIsPrepass);
@@ -124,15 +72,63 @@ public:
 		Encode a symbol to the bitstream. If the table is in prepass mode no actual bits will be emitted to the stream.
 	*/
 	void Encode(FHuffmanBitStreamWriter& Stream, int32 Symbol);
+private:
+	// This has a sentinel element at index NumValues!! And thus contains NumValues+1 elements
+	TArray<uint32> SymbolFrequencies;
+	TArray<int32> CodewordSymbols;
+
+	struct FCodeWord
+	{
+		int32 Bits;
+		int32 Length;
+	};
+
+	TArray<FCodeWord> CodeWords; // For every symbol the corresponding huffcode bits
+	int32 NumCodeWordsOfLength[HUFFMAN_MAX_CODE_LENGTH+1]; // Number of codewords with a given length, element 0 is unused
+	int32 NumSymbols;
+
+	void ResetFrequencies();
+	bool bIsPrepass;
+};
+
+class FHuffmanDecodeTable
+{
+public:
+	/**
+		Initialize a Huffman decode table based on the serialized table.
+	*/
+	void Initialize(FHuffmanBitStreamReader& Stream);
 
 	/**
 		Decode the next symbol from the bitstream.
 	*/
-	int32 Decode(FHuffmanBitStreamReader& Stream);
+	FORCEINLINE int32 Decode(FHuffmanBitStreamReader &Stream)
+	{
+		Stream.Refill();
+		const int32 Bits = Stream.PeekNoRefill(HUFFMAN_MAX_CODE_LENGTH);
+
+		const int32 Length = TableEntries[Bits].Length;
+		Stream.SkipNoRefill(Length);
+		return TableEntries[Bits].Symbol;
+	}
 
 	/**
-		Mark all symbols as being used at least once. This is useful if your prepass is not using the whole dataset but just a small fraction of it.
-		Symbols never occurred in the sample dataset can still be coded (but will obviously have a high coding cost).
+	Decode the next symbol from the bitstream without refilling the bit buffer.
 	*/
-	void MarkAllSymbols();
+	FORCEINLINE int32 DecodeNoRefill(FHuffmanBitStreamReader &Stream)
+	{
+		const int32 Bits = Stream.PeekNoRefill(HUFFMAN_MAX_CODE_LENGTH);
+		const int32 Length = TableEntries[Bits].Length;
+		Stream.SkipNoRefill(Length);
+		return TableEntries[Bits].Symbol;
+	}
+
+private:
+	struct FTableEntry
+	{
+		int8 Length;
+		int8 Symbol;
+	};
+
+	FTableEntry TableEntries[1u << HUFFMAN_MAX_CODE_LENGTH];
 };
diff --git a/Engine/Source/Runtime/Core/Public/UObject/RenderingObjectVersion.h b/Engine/Source/Runtime/Core/Public/UObject/RenderingObjectVersion.h
--- a/Engine/Source/Runtime/Core/Public/UObject/RenderingObjectVersion.h
+++ b/Engine/Source/Runtime/Core/Public/UObject/RenderingObjectVersion.h
@@ -77,6 +77,8 @@ struct CORE_API FRenderingObjectVersion
 
 		VirtualTexturedLightmaps,
 
+		GeometryCacheFastDecoder,
+
 		// -----<new versions can be added above this line>-------------------------------------------------
 		VersionPlusOne,
 		LatestVersion = VersionPlusOne - 1
