Fixed warning about variable shadowing on android #rb none

diff --git a/Engine/Plugins/Experimental/GeometryCache/Source/GeometryCache/Private/HuffmanTable.cpp b/Engine/Plugins/Experimental/GeometryCache/Source/GeometryCache/Private/HuffmanTable.cpp
--- a/Engine/Plugins/Experimental/GeometryCache/Source/GeometryCache/Private/HuffmanTable.cpp
+++ b/Engine/Plugins/Experimental/GeometryCache/Source/GeometryCache/Private/HuffmanTable.cpp
@@ -77,16 +77,16 @@ static void GenerateLengthLimitedHuffmanLengths(uint8* OutLengths, uint32* Frequ
 		Node*	Children[2];
 
 		// Traverse tree and every time a leaf node is visited the associated symbol code length is incremented by 1.
-		static void GenerateLengthLimitedHuffmanLengthsInner(Node* Node, uint8* OutLengths)
+		static void GenerateLengthLimitedHuffmanLengthsInner(Node* Node, uint8* OutCodeLengths)
 		{
 			if (Node->Symbol == INT32_MIN)
 			{
-				GenerateLengthLimitedHuffmanLengthsInner(Node->Children[0], OutLengths);
-				GenerateLengthLimitedHuffmanLengthsInner(Node->Children[1], OutLengths);
+				GenerateLengthLimitedHuffmanLengthsInner(Node->Children[0], OutCodeLengths);
+				GenerateLengthLimitedHuffmanLengthsInner(Node->Children[1], OutCodeLengths);
 			}
 			else
 			{
-				OutLengths[Node->Symbol]++;
+				OutCodeLengths[Node->Symbol]++;
 			}
 		}
 	};
