Zero initialize SymbolLists array in HuffmanTable.cpp to stop the static analysis from complaining about it #rb none

diff --git a/Engine/Plugins/Experimental/GeometryCache/Source/GeometryCache/Private/HuffmanTable.cpp b/Engine/Plugins/Experimental/GeometryCache/Source/GeometryCache/Private/HuffmanTable.cpp
--- a/Engine/Plugins/Experimental/GeometryCache/Source/GeometryCache/Private/HuffmanTable.cpp
+++ b/Engine/Plugins/Experimental/GeometryCache/Source/GeometryCache/Private/HuffmanTable.cpp
@@ -152,7 +152,7 @@ static void GenerateLengthLimitedHuffmanLengths(uint8* OutLengths, uint32* Frequ
 static void GenerateHuffmanCodes(uint16* OutCodes, const uint8* CodeLengths, int32 NumSymbols, int32 MaxCodeLength)
 {
 	int32 SymbolLengthHistogram[HUFFMAN_MAX_CODE_LENGTH + 1] = {};			// Histogram of code lengths
-	int16 SymbolLists[(HUFFMAN_MAX_CODE_LENGTH + 1) * HUFFMAN_MAX_CODES];	// List of symbols for every code length
+	int16 SymbolLists[(HUFFMAN_MAX_CODE_LENGTH + 1) * HUFFMAN_MAX_CODES] = {};	// List of symbols for every code length
 
 	// Fill histogram and generate symbol lists for each code length
 	for (int32 Symbol = 0; Symbol < NumSymbols; Symbol++)
