Sequencer: Fix start/end offset upgrade paths for audio, skeletal, geometry cache sections. Upgrade should be done based on the display resolution, not the legacy resolution.

diff --git a/Engine/Plugins/Experimental/GeometryCache/Source/GeometryCacheTracks/Private/MovieSceneGeometryCacheSection.cpp b/Engine/Plugins/Experimental/GeometryCache/Source/GeometryCacheTracks/Private/MovieSceneGeometryCacheSection.cpp
--- a/Engine/Plugins/Experimental/GeometryCache/Source/GeometryCacheTracks/Private/MovieSceneGeometryCacheSection.cpp
+++ b/Engine/Plugins/Experimental/GeometryCache/Source/GeometryCacheTracks/Private/MovieSceneGeometryCacheSection.cpp
@@ -46,18 +46,19 @@ void UMovieSceneGeometryCacheSection::PostLoad()
 {
 	Super::PostLoad();
 
-	FFrameRate LegacyFrameRate = GetLegacyConversionFrameRate();
+	FFrameRate DisplayRate = GetTypedOuter<UMovieScene>()->GetDisplayRate();
+	FFrameRate TickResolution = GetTypedOuter<UMovieScene>()->GetTickResolution();
 
 	if (Params.StartOffset_DEPRECATED != GeometryCacheDeprecatedMagicNumber)
 	{
-		Params.StartFrameOffset = UpgradeLegacyMovieSceneTime(this, LegacyFrameRate, Params.StartOffset_DEPRECATED).Value;
+		Params.StartFrameOffset = ConvertFrameTime(FFrameTime::FromDecimal(DisplayRate.AsDecimal() * Params.StartOffset_DEPRECATED), DisplayRate, TickResolution).FrameNumber;
 
 		Params.StartOffset_DEPRECATED = GeometryCacheDeprecatedMagicNumber;
 	}
 
 	if (Params.EndOffset_DEPRECATED != GeometryCacheDeprecatedMagicNumber)
 	{
-		Params.EndFrameOffset = UpgradeLegacyMovieSceneTime(this, LegacyFrameRate, Params.EndOffset_DEPRECATED).Value;
+		Params.EndFrameOffset = ConvertFrameTime(FFrameTime::FromDecimal(DisplayRate.AsDecimal() * Params.EndOffset_DEPRECATED), DisplayRate, TickResolution).FrameNumber;
 
 		Params.EndOffset_DEPRECATED = GeometryCacheDeprecatedMagicNumber;
 	}
diff --git a/Engine/Source/Runtime/MovieSceneTracks/Private/Sections/MovieSceneAudioSection.cpp b/Engine/Source/Runtime/MovieSceneTracks/Private/Sections/MovieSceneAudioSection.cpp
--- a/Engine/Source/Runtime/MovieSceneTracks/Private/Sections/MovieSceneAudioSection.cpp
+++ b/Engine/Source/Runtime/MovieSceneTracks/Private/Sections/MovieSceneAudioSection.cpp
@@ -118,8 +118,10 @@ void UMovieSceneAudioSection::PostLoad()
 
 	if (StartOffsetToUpgrade.IsSet())
 	{
-		FFrameNumber StartFrame = UpgradeLegacyMovieSceneTime(this, LegacyFrameRate, StartOffsetToUpgrade.GetValue());
-		StartFrameOffset = StartFrame.Value;
+		FFrameRate DisplayRate = GetTypedOuter<UMovieScene>()->GetDisplayRate();
+		FFrameRate TickResolution = GetTypedOuter<UMovieScene>()->GetTickResolution();
+
+		StartFrameOffset = ConvertFrameTime(FFrameTime::FromDecimal(DisplayRate.AsDecimal() * StartOffsetToUpgrade.GetValue()), DisplayRate, TickResolution).FrameNumber;
 	}
 }
 
diff --git a/Engine/Source/Runtime/MovieSceneTracks/Private/Sections/MovieSceneSkeletalAnimationSection.cpp b/Engine/Source/Runtime/MovieSceneTracks/Private/Sections/MovieSceneSkeletalAnimationSection.cpp
--- a/Engine/Source/Runtime/MovieSceneTracks/Private/Sections/MovieSceneSkeletalAnimationSection.cpp
+++ b/Engine/Source/Runtime/MovieSceneTracks/Private/Sections/MovieSceneSkeletalAnimationSection.cpp
@@ -110,18 +110,19 @@ void UMovieSceneSkeletalAnimationSection::PostLoad()
 		Params.SlotName = SlotName_DEPRECATED;
 	}
 
-	FFrameRate LegacyFrameRate = GetLegacyConversionFrameRate();
+	FFrameRate DisplayRate = GetTypedOuter<UMovieScene>()->GetDisplayRate();
+	FFrameRate TickResolution = GetTypedOuter<UMovieScene>()->GetTickResolution();
 
 	if (Params.StartOffset_DEPRECATED != SkeletalDeprecatedMagicNumber)
 	{
-		Params.StartFrameOffset = UpgradeLegacyMovieSceneTime(this, LegacyFrameRate, Params.StartOffset_DEPRECATED).Value;
+		Params.StartFrameOffset = ConvertFrameTime(FFrameTime::FromDecimal(DisplayRate.AsDecimal() * Params.StartOffset_DEPRECATED), DisplayRate, TickResolution).FrameNumber;
 
 		Params.StartOffset_DEPRECATED = SkeletalDeprecatedMagicNumber;
 	}
 
 	if (Params.EndOffset_DEPRECATED != SkeletalDeprecatedMagicNumber)
 	{
-		Params.EndFrameOffset = UpgradeLegacyMovieSceneTime(this, LegacyFrameRate, Params.EndOffset_DEPRECATED).Value;
+		Params.EndFrameOffset = ConvertFrameTime(FFrameTime::FromDecimal(DisplayRate.AsDecimal() * Params.EndOffset_DEPRECATED), DisplayRate, TickResolution).FrameNumber;
 
 		Params.EndOffset_DEPRECATED = SkeletalDeprecatedMagicNumber;
 	}
