Sequencer: Fix geometry cache reverse not using the complete sequence length to compute the reversed time.

diff --git a/Engine/Plugins/Experimental/GeometryCache/Source/GeometryCacheTracks/Private/MovieSceneGeometryCacheTemplate.cpp b/Engine/Plugins/Experimental/GeometryCache/Source/GeometryCacheTracks/Private/MovieSceneGeometryCacheTemplate.cpp
--- a/Engine/Plugins/Experimental/GeometryCache/Source/GeometryCacheTracks/Private/MovieSceneGeometryCacheTemplate.cpp
+++ b/Engine/Plugins/Experimental/GeometryCache/Source/GeometryCacheTracks/Private/MovieSceneGeometryCacheTemplate.cpp
@@ -144,7 +144,7 @@ float FMovieSceneGeometryCacheSectionTemplateParameters::MapTimeToAnimation(floa
 	AnimPosition += InFrameRate.AsSeconds(StartFrameOffset);
 	if (bReverse)
 	{
-		AnimPosition = (SeqLength - (AnimPosition - InFrameRate.AsSeconds(StartFrameOffset))) + InFrameRate.AsSeconds(StartFrameOffset);
+		AnimPosition = SequenceLength - AnimPosition;
 	}
 
 	return AnimPosition;
