Looping with offset fix also for Geometry Cache Tracks

diff --git a/Engine/Plugins/Experimental/GeometryCache/Source/GeometryCacheTracks/Private/MovieSceneGeometryCacheTemplate.cpp b/Engine/Plugins/Experimental/GeometryCache/Source/GeometryCacheTracks/Private/MovieSceneGeometryCacheTemplate.cpp
--- a/Engine/Plugins/Experimental/GeometryCache/Source/GeometryCacheTracks/Private/MovieSceneGeometryCacheTemplate.cpp
+++ b/Engine/Plugins/Experimental/GeometryCache/Source/GeometryCacheTracks/Private/MovieSceneGeometryCacheTemplate.cpp
@@ -127,7 +127,7 @@ float FMovieSceneGeometryCacheSectionTemplateParameters::MapTimeToAnimation(floa
 	FFrameTime AnimationLength = SequenceLength * InFrameRate;
 	int32 LengthInFrames = AnimationLength.FrameNumber.Value + (int)(AnimationLength.GetSubFrame() + 0.5f) + 1;
 	//we only play end if we are not looping, and assuming we are looping if Length is greater than default length;
-	bool bLooping = (SectionEndTime.Value - SectionStartTime.Value) > LengthInFrames;
+	bool bLooping = (SectionEndTime.Value - SectionStartTime.Value + StartFrameOffset + EndFrameOffset) > LengthInFrames;
 
 	InPosition = FMath::Clamp(InPosition, FFrameTime(SectionStartTime), FFrameTime(SectionEndTime - 1));
 
