Further fix geometry cache ray tracing stability issues by correctly dealing with segments especially when topology changes

diff --git a/Engine/Plugins/Experimental/GeometryCache/Source/GeometryCache/Private/GeometryCacheSceneProxy.cpp b/Engine/Plugins/Experimental/GeometryCache/Source/GeometryCache/Private/GeometryCacheSceneProxy.cpp
--- a/Engine/Plugins/Experimental/GeometryCache/Source/GeometryCache/Private/GeometryCacheSceneProxy.cpp
+++ b/Engine/Plugins/Experimental/GeometryCache/Source/GeometryCache/Private/GeometryCacheSceneProxy.cpp
@@ -159,7 +159,7 @@ FGeometryCacheSceneProxy::FGeometryCacheSceneProxy(UGeometryCacheComponent* Comp
 							Initializer.BaseVertexIndex = 0;
 							Initializer.VertexBufferStride = sizeof(FVector);
 							Initializer.VertexBufferByteOffset = 0;
-							Initializer.TotalPrimitiveCount = Section->IndexBuffer.NumIndices / 3;
+							Initializer.TotalPrimitiveCount = 0;
 							Initializer.VertexBufferElementType = VET_Float3;
 							Initializer.GeometryType = RTGT_Triangles;
 							Initializer.bFastBuild = false;
@@ -168,6 +168,7 @@ FGeometryCacheSceneProxy::FGeometryCacheSceneProxy(UGeometryCacheComponent* Comp
 							for (FGeometryCacheMeshBatchInfo& BatchInfo : Section->MeshData->BatchesInfo)
 							{
 								Segments.Add(FRayTracingGeometrySegment { BatchInfo.StartIndex / 3, BatchInfo.NumTriangles });
+								Initializer.TotalPrimitiveCount += BatchInfo.NumTriangles;
 							}
 
 							Initializer.Segments = Segments;
@@ -312,72 +313,72 @@ void FGeometryCacheSceneProxy::CreateMeshBatch(
 	FDynamicPrimitiveUniformBuffer& DynamicPrimitiveUniformBuffer,
 	FMeshBatch& Mesh) const
 {
-	FGeometryCacheVertexFactoryUserData &UserData = UserDataWrapper.Data;
+	FGeometryCacheVertexFactoryUserData& UserData = UserDataWrapper.Data;
 
 	UserData.MeshExtension = FVector::OneVector;
 	UserData.MeshOrigin = FVector::ZeroVector;
 
-			const bool bHasMotionVectors = (
-				TrackProxy->MeshData->VertexInfo.bHasMotionVectors &&
-				TrackProxy->NextFrameMeshData->VertexInfo.bHasMotionVectors &&
-				TrackProxy->MeshData->Positions.Num() == TrackProxy->MeshData->MotionVectors.Num())
-				&& (TrackProxy->NextFrameMeshData->Positions.Num() == TrackProxy->NextFrameMeshData->MotionVectors.Num());
+	const bool bHasMotionVectors = (
+		TrackProxy->MeshData->VertexInfo.bHasMotionVectors &&
+		TrackProxy->NextFrameMeshData->VertexInfo.bHasMotionVectors &&
+		TrackProxy->MeshData->Positions.Num() == TrackProxy->MeshData->MotionVectors.Num())
+		&& (TrackProxy->NextFrameMeshData->Positions.Num() == TrackProxy->NextFrameMeshData->MotionVectors.Num());
 
-				if (!bHasMotionVectors)
-				{
-					UserData.MotionBlurDataExtension = FVector::OneVector;
-					UserData.MotionBlurDataOrigin = FVector::ZeroVector;
-					UserData.MotionBlurPositionScale = 0.0f;
-				}
-				else
-				{
-					UserData.MotionBlurDataExtension = FVector::OneVector * PlaybackSpeed;
-					UserData.MotionBlurDataOrigin = FVector::ZeroVector;
-					UserData.MotionBlurPositionScale = 1.0f;
-				}
+	if (!bHasMotionVectors)
+	{
+		UserData.MotionBlurDataExtension = FVector::OneVector;
+		UserData.MotionBlurDataOrigin = FVector::ZeroVector;
+		UserData.MotionBlurPositionScale = 0.0f;
+	}
+	else
+	{
+		UserData.MotionBlurDataExtension = FVector::OneVector * PlaybackSpeed;
+		UserData.MotionBlurDataOrigin = FVector::ZeroVector;
+		UserData.MotionBlurPositionScale = 1.0f;
+	}
 
-				if (IsRayTracingEnabled())
-				{
-					// No vertex manipulation is allowed in the vertex shader
-					// Otherwise we need an additional compute shader pass to execute the vertex shader and dump to a staging buffer
-					check(UserData.MeshExtension == FVector::OneVector);
-					check(UserData.MeshOrigin == FVector::ZeroVector);
-				}
+	if (IsRayTracingEnabled())
+	{
+		// No vertex manipulation is allowed in the vertex shader
+		// Otherwise we need an additional compute shader pass to execute the vertex shader and dump to a staging buffer
+		check(UserData.MeshExtension == FVector::OneVector);
+		check(UserData.MeshOrigin == FVector::ZeroVector);
+	}
 
-				UserData.PositionBuffer = &TrackProxy->PositionBuffers[TrackProxy->CurrentPositionBufferIndex % 2];
-				UserData.MotionBlurDataBuffer = &TrackProxy->PositionBuffers[(TrackProxy->CurrentPositionBufferIndex+1) % 2];
+	UserData.PositionBuffer = &TrackProxy->PositionBuffers[TrackProxy->CurrentPositionBufferIndex % 2];
+	UserData.MotionBlurDataBuffer = &TrackProxy->PositionBuffers[(TrackProxy->CurrentPositionBufferIndex+1) % 2];
 
-				FGeometryCacheVertexFactoryUniformBufferParameters UniformBufferParameters;
+	FGeometryCacheVertexFactoryUniformBufferParameters UniformBufferParameters;
 
-				UniformBufferParameters.MeshOrigin = UserData.MeshOrigin;
-				UniformBufferParameters.MeshExtension = UserData.MeshExtension;
-				UniformBufferParameters.MotionBlurDataOrigin = UserData.MotionBlurDataOrigin;
-				UniformBufferParameters.MotionBlurDataExtension = UserData.MotionBlurDataExtension;
-				UniformBufferParameters.MotionBlurPositionScale = UserData.MotionBlurPositionScale;
+	UniformBufferParameters.MeshOrigin = UserData.MeshOrigin;
+	UniformBufferParameters.MeshExtension = UserData.MeshExtension;
+	UniformBufferParameters.MotionBlurDataOrigin = UserData.MotionBlurDataOrigin;
+	UniformBufferParameters.MotionBlurDataExtension = UserData.MotionBlurDataExtension;
+	UniformBufferParameters.MotionBlurPositionScale = UserData.MotionBlurPositionScale;
 
-				UserData.UniformBuffer = FGeometryCacheVertexFactoryUniformBufferParametersRef::CreateUniformBufferImmediate(UniformBufferParameters, UniformBuffer_SingleFrame);
-				TrackProxy->VertexFactory.CreateManualVertexFetchUniformBuffer(UserData.PositionBuffer, UserData.MotionBlurDataBuffer, UserData);
+	UserData.UniformBuffer = FGeometryCacheVertexFactoryUniformBufferParametersRef::CreateUniformBufferImmediate(UniformBufferParameters, UniformBuffer_SingleFrame);
+	TrackProxy->VertexFactory.CreateManualVertexFetchUniformBuffer(UserData.PositionBuffer, UserData.MotionBlurDataBuffer, UserData);
 
-				// Draw the mesh.
-				FMeshBatchElement& BatchElement = Mesh.Elements[0];
-				BatchElement.IndexBuffer = &TrackProxy->IndexBuffer;
-				Mesh.VertexFactory = &TrackProxy->VertexFactory;
-				Mesh.SegmentIndex = 0;
+	// Draw the mesh.
+	FMeshBatchElement& BatchElement = Mesh.Elements[0];
+	BatchElement.IndexBuffer = &TrackProxy->IndexBuffer;
+	Mesh.VertexFactory = &TrackProxy->VertexFactory;
+	Mesh.SegmentIndex = 0;
 
-				const FMatrix& LocalToWorldTransform = TrackProxy->WorldMatrix * GetLocalToWorld();
+	const FMatrix& LocalToWorldTransform = TrackProxy->WorldMatrix * GetLocalToWorld();
 
 	DynamicPrimitiveUniformBuffer.Set(LocalToWorldTransform, LocalToWorldTransform, GetBounds(), GetLocalBounds(), true, false, DrawsVelocity(), false);
 	BatchElement.PrimitiveUniformBuffer = DynamicPrimitiveUniformBuffer.UniformBuffer.GetUniformBufferRHI();
 
-				BatchElement.FirstIndex = BatchInfo.StartIndex;
-				BatchElement.NumPrimitives = BatchInfo.NumTriangles;
-				BatchElement.MinVertexIndex = 0;
-				BatchElement.MaxVertexIndex = TrackProxy->MeshData->Positions.Num() - 1;
-				BatchElement.VertexFactoryUserData = &UserDataWrapper.Data;
-				Mesh.ReverseCulling = IsLocalToWorldDeterminantNegative();
-				Mesh.Type = PT_TriangleList;
-				Mesh.DepthPriorityGroup = SDPG_World;
-				Mesh.bCanApplyViewModeOverrides = false;
+	BatchElement.FirstIndex = BatchInfo.StartIndex;
+	BatchElement.NumPrimitives = BatchInfo.NumTriangles;
+	BatchElement.MinVertexIndex = 0;
+	BatchElement.MaxVertexIndex = TrackProxy->MeshData->Positions.Num() - 1;
+	BatchElement.VertexFactoryUserData = &UserDataWrapper.Data;
+	Mesh.ReverseCulling = IsLocalToWorldDeterminantNegative();
+	Mesh.Type = PT_TriangleList;
+	Mesh.DepthPriorityGroup = SDPG_World;
+	Mesh.bCanApplyViewModeOverrides = false;
 }
 
 void FGeometryCacheSceneProxy::GetDynamicMeshElements(const TArray<const FSceneView*>& Views, const FSceneViewFamily& ViewFamily, uint32 VisibilityMap, FMeshElementCollector& Collector) const
@@ -553,7 +554,18 @@ void FGeometryCacheSceneProxy::UpdateAnimation(float NewTime, bool bNewLooping,
 				const int PositionBufferIndex = Section->CurrentPositionBufferIndex != -1 ? Section->CurrentPositionBufferIndex % 2 : 0;
 
 				Section->RayTracingGeometry.Initializer.PositionVertexBuffer = Section->PositionBuffers[PositionBufferIndex].VertexBufferRHI;
-				Section->RayTracingGeometry.Initializer.TotalPrimitiveCount = Section->IndexBuffer.NumIndices / 3;
+				Section->RayTracingGeometry.Initializer.IndexBuffer = Section->IndexBuffer.IndexBufferRHI;
+				Section->RayTracingGeometry.Initializer.TotalPrimitiveCount = 0;
+				
+				TArray<FRayTracingGeometrySegment> Segments;
+				for (FGeometryCacheMeshBatchInfo& BatchInfo : Section->MeshData->BatchesInfo)
+				{
+					Segments.Add(FRayTracingGeometrySegment { BatchInfo.StartIndex / 3, BatchInfo.NumTriangles });
+					Section->RayTracingGeometry.Initializer.TotalPrimitiveCount += BatchInfo.NumTriangles;
+				}
+
+				Section->RayTracingGeometry.Initializer.Segments = Segments;
+							
 				Section->RayTracingGeometry.UpdateRHI();
 			}
 		}
