Fixed driver crash when scrubbing in the sequencer with ray tracing enabled with geometry cache. The index buffer was modified by the render thread while the RHI thread was reading it for building the ray tracing acceleration structure.

(ushell-p4-cherrypick of 9551390 by [REDACTED_PERSON])

diff --git a/Engine/Plugins/Experimental/GeometryCache/Source/GeometryCache/Private/GeometryCacheSceneProxy.cpp b/Engine/Plugins/Experimental/GeometryCache/Source/GeometryCache/Private/GeometryCacheSceneProxy.cpp
--- a/Engine/Plugins/Experimental/GeometryCache/Source/GeometryCache/Private/GeometryCacheSceneProxy.cpp
+++ b/Engine/Plugins/Experimental/GeometryCache/Source/GeometryCache/Private/GeometryCacheSceneProxy.cpp
@@ -1040,7 +1040,7 @@ void FGeomCacheIndexBuffer::InitRHI()
 {
 	FRHIResourceCreateInfo CreateInfo;
 	void* Buffer = nullptr;
-	IndexBufferRHI = RHICreateAndLockIndexBuffer(sizeof(uint32), NumIndices * sizeof(uint32), BUF_Static | BUF_ShaderResource, CreateInfo, Buffer);
+	IndexBufferRHI = RHICreateAndLockIndexBuffer(sizeof(uint32), NumIndices * sizeof(uint32), BUF_Dynamic | BUF_ShaderResource, CreateInfo, Buffer);
 	RHIUnlockIndexBuffer(IndexBufferRHI);
 }
 
@@ -1057,7 +1057,7 @@ void FGeomCacheIndexBuffer::Update(const TArray<uint32> &Indices)
 	{
 		NumIndices = Indices.Num();
 		FRHIResourceCreateInfo CreateInfo;
-		IndexBufferRHI = RHICreateAndLockIndexBuffer(sizeof(uint32), NumIndices * sizeof(uint32), BUF_Static | BUF_ShaderResource, CreateInfo, Buffer);
+		IndexBufferRHI = RHICreateAndLockIndexBuffer(sizeof(uint32), NumIndices * sizeof(uint32), BUF_Dynamic | BUF_ShaderResource, CreateInfo, Buffer);
 	}
 	else
 	{
@@ -1077,7 +1077,7 @@ void FGeomCacheIndexBuffer::UpdateSizeOnly(int32 NewNumIndices)
 	if (NewNumIndices > NumIndices)
 	{
 		FRHIResourceCreateInfo CreateInfo;
-		IndexBufferRHI = RHICreateIndexBuffer(sizeof(uint32), NewNumIndices * sizeof(uint32), BUF_Static | BUF_ShaderResource, CreateInfo);
+		IndexBufferRHI = RHICreateIndexBuffer(sizeof(uint32), NewNumIndices * sizeof(uint32), BUF_Dynamic | BUF_ShaderResource, CreateInfo);
 		NumIndices = NewNumIndices;
 	}
 }
