Alembic GeometryCache Actor color stays the same when selected while viewing in Wireframe
#fix Copy selection behaviour from static mesh render proxy

diff --git a/Engine/Plugins/Experimental/GeometryCache/Source/GeometryCache/Private/GeometryCacheSceneProxy.cpp b/Engine/Plugins/Experimental/GeometryCache/Source/GeometryCache/Private/GeometryCacheSceneProxy.cpp
--- a/Engine/Plugins/Experimental/GeometryCache/Source/GeometryCache/Private/GeometryCacheSceneProxy.cpp
+++ b/Engine/Plugins/Experimental/GeometryCache/Source/GeometryCache/Private/GeometryCacheSceneProxy.cpp
@@ -394,10 +394,20 @@ void FGeometryCacheSceneProxy::GetDynamicMeshElements(const TArray<const FSceneV
 	FColoredMaterialRenderProxy* WireframeMaterialInstance = nullptr;
 	if (bWireframe)
 	{
+		const FEngineShowFlags& EngineShowFlags = ViewFamily.EngineShowFlags;
+		const bool bLevelColorationEnabled = EngineShowFlags.LevelColoration;
+		const bool bPropertyColorationEnabled = EngineShowFlags.PropertyColoration;
+
+		FLinearColor ViewWireframeColor(bLevelColorationEnabled ? GetLevelColor() : GetWireframeColor());
+		if (bPropertyColorationEnabled)
+		{
+			ViewWireframeColor = GetPropertyColor();
+		}
+
 		WireframeMaterialInstance = new FColoredMaterialRenderProxy(
 			GEngine->WireframeMaterial ? GEngine->WireframeMaterial->GetRenderProxy() : nullptr,
-			FLinearColor(0, 0.5f, 1.f)
-			);
+			GetSelectionColor(ViewWireframeColor, !(GIsEditor && EngineShowFlags.Selection) || IsSelected(), IsHovered(), false)
+		);
 
 		Collector.RegisterOneFrameMaterialProxy(WireframeMaterialInstance);
 	}
