Geometry Cache crashing when loading async in a cooked build
#fix allow adding chunks from async loading thread

diff --git a/Engine/Plugins/Experimental/GeometryCache/Source/GeometryCache/Private/GeometryCacheStreamingManager.cpp b/Engine/Plugins/Experimental/GeometryCache/Source/GeometryCache/Private/GeometryCacheStreamingManager.cpp
--- a/Engine/Plugins/Experimental/GeometryCache/Source/GeometryCache/Private/GeometryCacheStreamingManager.cpp
+++ b/Engine/Plugins/Experimental/GeometryCache/Source/GeometryCache/Private/GeometryCacheStreamingManager.cpp
@@ -225,7 +225,7 @@ void FGeometryCacheStreamingManager::NotifyLevelOffset(class ULevel* Level, cons
 
 void FGeometryCacheStreamingManager::AddGeometryCache(UGeometryCacheTrackStreamable* Cache)
 {
-	check(IsInGameThread());
+	check(IsInGameThread() || IsInAsyncLoadingThread());
 	FStreamingGeometryCacheData* &CacheData = StreamingGeometryCaches.FindOrAdd(Cache);
 	if (CacheData == nullptr)
 	{
