Fixed incorrect ray tracing segments creation in geometry cache

(ushell-p4-cherrypick of 9839937 by [REDACTED_PERSON])

diff --git a/Engine/Plugins/Experimental/GeometryCache/Source/GeometryCache/Private/GeometryCacheSceneProxy.cpp b/Engine/Plugins/Experimental/GeometryCache/Source/GeometryCache/Private/GeometryCacheSceneProxy.cpp
--- a/Engine/Plugins/Experimental/GeometryCache/Source/GeometryCache/Private/GeometryCacheSceneProxy.cpp
+++ b/Engine/Plugins/Experimental/GeometryCache/Source/GeometryCache/Private/GeometryCacheSceneProxy.cpp
@@ -570,7 +570,7 @@ void FGeometryCacheSceneProxy::UpdateAnimation(float NewTime, bool bNewLooping,
 				Section->RayTracingGeometry.Initializer.IndexBuffer = Section->IndexBuffer.IndexBufferRHI;
 				Section->RayTracingGeometry.Initializer.TotalPrimitiveCount = 0;
 				
-				TArray<FRayTracingGeometrySegment>& Segments = Section->RayTracingGeometry.Initializer.Segments;
+				TArray<FRayTracingGeometrySegment> Segments;
 
 				for (FGeometryCacheMeshBatchInfo& BatchInfo : Section->MeshData->BatchesInfo)
 				{
