Small optimization processing ray tracing segments

(ushell-p4-cherrypick of 9840613 by [REDACTED_PERSON])

diff --git a/Engine/Plugins/Experimental/GeometryCache/Source/GeometryCache/Private/GeometryCacheSceneProxy.cpp b/Engine/Plugins/Experimental/GeometryCache/Source/GeometryCache/Private/GeometryCacheSceneProxy.cpp
--- a/Engine/Plugins/Experimental/GeometryCache/Source/GeometryCache/Private/GeometryCacheSceneProxy.cpp
+++ b/Engine/Plugins/Experimental/GeometryCache/Source/GeometryCache/Private/GeometryCacheSceneProxy.cpp
@@ -570,7 +570,8 @@ void FGeometryCacheSceneProxy::UpdateAnimation(float NewTime, bool bNewLooping,
 				Section->RayTracingGeometry.Initializer.IndexBuffer = Section->IndexBuffer.IndexBufferRHI;
 				Section->RayTracingGeometry.Initializer.TotalPrimitiveCount = 0;
 				
-				TArray<FRayTracingGeometrySegment> Segments;
+				TArray<FRayTracingGeometrySegment>& Segments = Section->RayTracingGeometry.Initializer.Segments;
+				Segments.Reset();
 
 				for (FGeometryCacheMeshBatchInfo& BatchInfo : Section->MeshData->BatchesInfo)
 				{
@@ -582,8 +583,6 @@ void FGeometryCacheSceneProxy::UpdateAnimation(float NewTime, bool bNewLooping,
 					Segments.Add(Segment);
 					Section->RayTracingGeometry.Initializer.TotalPrimitiveCount += BatchInfo.NumTriangles;
 				}
-
-				Section->RayTracingGeometry.Initializer.Segments = Segments;
 							
 				Section->RayTracingGeometry.UpdateRHI();
 			}
