[CrashReport] UE4Editor-GeometryCache!DispatchCheckVerify<void,<lambda_bfafc0bbea4685bda9fbd9fff0b95261> >(<lambda_bfafc0bbea4685bda9fbd9fff0b95261> &&) [AssertionMacros.h:162]
#fix Adding ensure for suspected reason of crash, as none of the reports have any repro steps or description of potential cause

diff --git a/Engine/Plugins/Experimental/GeometryCache/Source/GeometryCache/Private/GeometryCacheTrackStreamable.cpp b/Engine/Plugins/Experimental/GeometryCache/Source/GeometryCache/Private/GeometryCacheTrackStreamable.cpp
--- a/Engine/Plugins/Experimental/GeometryCache/Source/GeometryCache/Private/GeometryCacheTrackStreamable.cpp
+++ b/Engine/Plugins/Experimental/GeometryCache/Source/GeometryCache/Private/GeometryCacheTrackStreamable.cpp
@@ -508,6 +508,7 @@ const FVisibilitySample& UGeometryCacheTrackStreamable::GetVisibilitySample(floa
 			return Sample;
 		}
 	}
+	ensure(VisibilitySamples.Num());
 	return VisibilitySamples.Last();
 }
 
