Fixed cast shadows flag was ignored for geometry cache meshes with ray tracing shadows enabled.

diff --git a/Engine/Plugins/Experimental/GeometryCache/Source/GeometryCache/Private/GeometryCacheSceneProxy.cpp b/Engine/Plugins/Experimental/GeometryCache/Source/GeometryCache/Private/GeometryCacheSceneProxy.cpp
--- a/Engine/Plugins/Experimental/GeometryCache/Source/GeometryCache/Private/GeometryCacheSceneProxy.cpp
+++ b/Engine/Plugins/Experimental/GeometryCache/Source/GeometryCache/Private/GeometryCacheSceneProxy.cpp
@@ -504,6 +504,7 @@ void FGeometryCacheSceneProxy::GetDynamicRayTracingInstances(FRayTracingMaterial
 			CreateMeshBatch(TrackProxy, BatchInfo, UserDataWrapper, DynamicPrimitiveUniformBuffer, MeshBatch);
 
 			MeshBatch.MaterialRenderProxy = TrackProxy->Materials[SegmentIndex]->GetRenderProxy();
+			MeshBatch.CastRayTracedShadow = IsShadowCast(Context.ReferenceView);
 
 			RayTracingInstance.Materials.Add(MeshBatch);
 		}
