Sequencer: add initial first loop start offset to animation and geometry cache sections

Now when the user trims the start of a section by dragging its left edge, it only trims the first loop of that animation, so that extending the section by its right edge will make it loop seamlessly.

The old start/end offsets are still available for people who want to loop a subset of the animation, but they have to set them in the properties menu (before this change it was actually inconsistent since dragging the left edge was setting the start offset, but dragging the right edge was just resizing the section and making a loop).

diff --git a/Engine/Plugins/Experimental/GeometryCache/Source/GeometryCacheSequencer/Classes/GeometryCacheTrackEditor.cpp b/Engine/Plugins/Experimental/GeometryCache/Source/GeometryCacheSequencer/Classes/GeometryCacheTrackEditor.cpp
--- a/Engine/Plugins/Experimental/GeometryCache/Source/GeometryCacheSequencer/Classes/GeometryCacheTrackEditor.cpp
+++ b/Engine/Plugins/Experimental/GeometryCache/Source/GeometryCacheSequencer/Classes/GeometryCacheTrackEditor.cpp
@@ -67,7 +67,7 @@ static UGeometryCacheComponent* AcquireGeometryCacheFromObjectGuid(const FGuid&
 FGeometryCacheSection::FGeometryCacheSection(UMovieSceneSection& InSection, TWeakPtr<ISequencer> InSequencer)
 	: Section(*CastChecked<UMovieSceneGeometryCacheSection>(&InSection))
 	, Sequencer(InSequencer)
-	, InitialStartOffsetDuringResize(0)
+	, InitialFirstLoopStartOffsetDuringResize(0)
 	, InitialStartTimeDuringResize(0)
 { }
 
@@ -113,14 +113,15 @@ int32 FGeometryCacheSection::OnPaintSection(FSequencerSectionPainter& Painter) c
 	FFrameRate TickResolution = TimeToPixelConverter.GetTickResolution();
 
 	// Add lines where the animation starts and ends/loops
-	float AnimPlayRate = FMath::IsNearlyZero(Section.Params.PlayRate) ? 1.0f : Section.Params.PlayRate;
-	float Duration = Section.Params.GetSequenceLength();
-	float SeqLength = Duration - TickResolution.AsSeconds(Section.Params.StartFrameOffset + Section.Params.EndFrameOffset) / AnimPlayRate;
+	const float AnimPlayRate = FMath::IsNearlyZero(Section.Params.PlayRate) ? 1.0f : Section.Params.PlayRate;
+	const float Duration = Section.Params.GetSequenceLength();
+	const float SeqLength = Duration - TickResolution.AsSeconds(Section.Params.StartFrameOffset + Section.Params.EndFrameOffset) / AnimPlayRate;
+	const float FirstLoopSeqLength = SeqLength - TickResolution.AsSeconds(Section.Params.FirstLoopStartFrameOffset) / AnimPlayRate;
 
 	if (!FMath::IsNearlyZero(SeqLength, KINDA_SMALL_NUMBER) && SeqLength > 0)
 	{
 		float MaxOffset = Section.GetRange().Size<FFrameTime>() / TickResolution;
-		float OffsetTime = SeqLength;
+		float OffsetTime = FirstLoopSeqLength;
 		float StartTime = Section.GetInclusiveStartFrame() / TickResolution;
 
 		while (OffsetTime < MaxOffset)
@@ -200,7 +201,7 @@ int32 FGeometryCacheSection::OnPaintSection(FSequencerSectionPainter& Painter) c
 
 void FGeometryCacheSection::BeginResizeSection()
 {
-	InitialStartOffsetDuringResize = Section.Params.StartFrameOffset;
+	InitialFirstLoopStartOffsetDuringResize = Section.Params.FirstLoopStartFrameOffset;
 	InitialStartTimeDuringResize = Section.HasStartFrame() ? Section.GetInclusiveStartFrame() : 0;
 }
 
@@ -212,17 +213,23 @@ void FGeometryCacheSection::ResizeSection(ESequencerSectionResizeMode ResizeMode
 		FFrameRate FrameRate = Section.GetTypedOuter<UMovieScene>()->GetTickResolution();
 		FFrameNumber StartOffset = FrameRate.AsFrameNumber((ResizeTime - InitialStartTimeDuringResize) / FrameRate * Section.Params.PlayRate);
 
-		StartOffset += InitialStartOffsetDuringResize;
+		StartOffset += InitialFirstLoopStartOffsetDuringResize;
 
-		// Ensure start offset is not less than 0 and adjust ResizeTime
 		if (StartOffset < 0)
 		{
+			// Ensure start offset is not less than 0 and adjust ResizeTime
 			ResizeTime = ResizeTime - StartOffset;
 
 			StartOffset = FFrameNumber(0);
 		}
+		else
+		{
+			// If the start offset exceeds the length of one loop, trim it back.
+			const FFrameNumber SeqLength = FrameRate.AsFrameNumber(Section.Params.GetSequenceLength()) - Section.Params.StartFrameOffset - Section.Params.EndFrameOffset;
+			StartOffset = StartOffset % SeqLength;
+		}
 
-		Section.Params.StartFrameOffset = StartOffset;
+		Section.Params.FirstLoopStartFrameOffset = StartOffset;
 	}
 
 	ISequencerSection::ResizeSection(ResizeMode, ResizeTime);
@@ -238,17 +245,23 @@ void FGeometryCacheSection::SlipSection(FFrameNumber SlipTime)
 	FFrameRate FrameRate = Section.GetTypedOuter<UMovieScene>()->GetTickResolution();
 	FFrameNumber StartOffset = FrameRate.AsFrameNumber((SlipTime - InitialStartTimeDuringResize) / FrameRate * Section.Params.PlayRate);
 
-	StartOffset += InitialStartOffsetDuringResize;
+	StartOffset += InitialFirstLoopStartOffsetDuringResize;
 
-	// Ensure start offset is not less than 0 and adjust ResizeTime
 	if (StartOffset < 0)
 	{
+		// Ensure start offset is not less than 0 and adjust ResizeTime
 		SlipTime = SlipTime - StartOffset;
 
 		StartOffset = FFrameNumber(0);
 	}
+	else
+	{
+		// If the start offset exceeds the length of one loop, trim it back.
+		const FFrameNumber SeqLength = FrameRate.AsFrameNumber(Section.Params.GetSequenceLength()) - Section.Params.StartFrameOffset - Section.Params.EndFrameOffset;
+		StartOffset = StartOffset % SeqLength;
+	}
 
-	Section.Params.StartFrameOffset = StartOffset;
+	Section.Params.FirstLoopStartFrameOffset = StartOffset;
 
 	ISequencerSection::SlipSection(SlipTime);
 }
diff --git a/Engine/Plugins/Experimental/GeometryCache/Source/GeometryCacheSequencer/Classes/GeometryCacheTrackEditor.h b/Engine/Plugins/Experimental/GeometryCache/Source/GeometryCacheSequencer/Classes/GeometryCacheTrackEditor.h
--- a/Engine/Plugins/Experimental/GeometryCache/Source/GeometryCacheSequencer/Classes/GeometryCacheTrackEditor.h
+++ b/Engine/Plugins/Experimental/GeometryCache/Source/GeometryCacheSequencer/Classes/GeometryCacheTrackEditor.h
@@ -95,8 +95,8 @@ private:
 	/** Used to draw animation frame, need selection state and local time*/
 	TWeakPtr<ISequencer> Sequencer;
 
-	/** Cached start offset value valid only during resize */
-	FFrameNumber InitialStartOffsetDuringResize;
+	/** Cached first loop start offset value valid only during resize */
+	FFrameNumber InitialFirstLoopStartOffsetDuringResize;
 	
 	/** Cached start time valid only during resize */
 	FFrameNumber InitialStartTimeDuringResize;
diff --git a/Engine/Plugins/Experimental/GeometryCache/Source/GeometryCacheTracks/Classes/MovieSceneGeometryCacheSection.h b/Engine/Plugins/Experimental/GeometryCache/Source/GeometryCacheTracks/Classes/MovieSceneGeometryCacheSection.h
--- a/Engine/Plugins/Experimental/GeometryCache/Source/GeometryCacheTracks/Classes/MovieSceneGeometryCacheSection.h
+++ b/Engine/Plugins/Experimental/GeometryCache/Source/GeometryCacheTracks/Classes/MovieSceneGeometryCacheSection.h
@@ -24,6 +24,10 @@ USTRUCT()
 	UPROPERTY(EditAnywhere, Category = "GeometryCache", DisplayName = "Geometry Cache")
 	UGeometryCache* GeometryCacheAsset;
 
+	/** The offset for the first loop of the animation clip */
+	UPROPERTY(EditAnywhere, Category = "GeometryCache")
+	FFrameNumber FirstLoopStartFrameOffset;
+
 	/** The offset into the beginning of the animation clip */
 	UPROPERTY(EditAnywhere, Category = "GeometryCache")
 	FFrameNumber StartFrameOffset;
diff --git a/Engine/Plugins/Experimental/GeometryCache/Source/GeometryCacheTracks/Private/MovieSceneGeometryCacheSection.cpp b/Engine/Plugins/Experimental/GeometryCache/Source/GeometryCacheTracks/Private/MovieSceneGeometryCacheSection.cpp
--- a/Engine/Plugins/Experimental/GeometryCache/Source/GeometryCacheTracks/Private/MovieSceneGeometryCacheSection.cpp
+++ b/Engine/Plugins/Experimental/GeometryCache/Source/GeometryCacheTracks/Private/MovieSceneGeometryCacheSection.cpp
@@ -39,7 +39,7 @@ UMovieSceneGeometryCacheSection::UMovieSceneGeometryCacheSection(const FObjectIn
 
 TOptional<FFrameTime> UMovieSceneGeometryCacheSection::GetOffsetTime() const
 {
-	return TOptional<FFrameTime>(Params.StartFrameOffset);
+	return TOptional<FFrameTime>(Params.FirstLoopStartFrameOffset);
 }
 
 void UMovieSceneGeometryCacheSection::PostLoad()
@@ -86,14 +86,17 @@ FMovieSceneEvalTemplatePtr UMovieSceneGeometryCacheSection::GenerateTemplate() c
 	return FMovieSceneGeometryCacheSectionTemplate(*this);
 }
 
-FFrameNumber GetStartOffsetAtTrimTime(FQualifiedFrameTime TrimTime, const FMovieSceneGeometryCacheParams& Params, FFrameNumber StartFrame, FFrameRate FrameRate)
+FFrameNumber GetFirstLoopStartOffsetAtTrimTime(FQualifiedFrameTime TrimTime, const FMovieSceneGeometryCacheParams& Params, FFrameNumber StartFrame, FFrameRate FrameRate)
 {
-	float AnimPlayRate = FMath::IsNearlyZero(Params.PlayRate) ? 1.0f : Params.PlayRate;
-	float AnimPosition = (TrimTime.Time - StartFrame) / TrimTime.Rate * AnimPlayRate;
-	float SeqLength = Params.GetSequenceLength() - FrameRate.AsSeconds(Params.StartFrameOffset + Params.EndFrameOffset) / AnimPlayRate;
+	const float AnimPlayRate = FMath::IsNearlyZero(Params.PlayRate) ? 1.0f : Params.PlayRate;
+	const float AnimPosition = (TrimTime.Time - StartFrame) / TrimTime.Rate * AnimPlayRate;
+	const float SeqLength = Params.GetSequenceLength() - FrameRate.AsSeconds(Params.StartFrameOffset + Params.EndFrameOffset) / AnimPlayRate;
 
 	FFrameNumber NewOffset = FrameRate.AsFrameNumber(FMath::Fmod(AnimPosition, SeqLength));
-	NewOffset += Params.StartFrameOffset;
+	NewOffset += Params.FirstLoopStartFrameOffset;
+
+	const FFrameNumber SeqLengthInFrames = FrameRate.AsFrameNumber(SeqLength);
+	NewOffset = NewOffset % SeqLengthInFrames;
 
 	return NewOffset;
 }
@@ -119,7 +122,7 @@ void UMovieSceneGeometryCacheSection::TrimSection(FQualifiedFrameTime TrimTime,
 		{
 			FFrameRate FrameRate = GetTypedOuter<UMovieScene>()->GetTickResolution();
 
-			Params.StartFrameOffset = HasStartFrame() ? GetStartOffsetAtTrimTime(TrimTime, Params, GetInclusiveStartFrame(), FrameRate) : 0;
+			Params.FirstLoopStartFrameOffset = HasStartFrame() ? GetFirstLoopStartOffsetAtTrimTime(TrimTime, Params, GetInclusiveStartFrame(), FrameRate) : 0;
 		}
 
 		Super::TrimSection(TrimTime, bTrimLeft);
@@ -130,13 +133,13 @@ UMovieSceneSection* UMovieSceneGeometryCacheSection::SplitSection(FQualifiedFram
 {
 	FFrameRate FrameRate = GetTypedOuter<UMovieScene>()->GetTickResolution();
 
-	const FFrameNumber NewOffset = HasStartFrame() ? GetStartOffsetAtTrimTime(SplitTime, Params, GetInclusiveStartFrame(), FrameRate) : 0;
+	const FFrameNumber NewOffset = HasStartFrame() ? GetFirstLoopStartOffsetAtTrimTime(SplitTime, Params, GetInclusiveStartFrame(), FrameRate) : 0;
 
 	UMovieSceneSection* NewSection = Super::SplitSection(SplitTime);
 	if (NewSection != nullptr)
 	{
 		UMovieSceneGeometryCacheSection* NewGeometrySection = Cast<UMovieSceneGeometryCacheSection>(NewSection);
-		NewGeometrySection->Params.StartFrameOffset = NewOffset;
+		NewGeometrySection->Params.FirstLoopStartFrameOffset = NewOffset;
 	}
 	return NewSection;
 }
@@ -152,16 +155,27 @@ void UMovieSceneGeometryCacheSection::GetSnapTimes(TArray<FFrameNumber>& OutSnap
 
 	const float AnimPlayRate = FMath::IsNearlyZero(Params.PlayRate) ? 1.0f : Params.PlayRate;
 	const float SeqLengthSeconds = Params.GetSequenceLength() - FrameRate.AsSeconds(Params.StartFrameOffset + Params.EndFrameOffset) / AnimPlayRate;
+	const float FirstLoopSeqLengthInSeconds = SeqLengthSeconds - FrameRate.AsSeconds(Params.FirstLoopStartFrameOffset) / AnimPlayRate;
 
-	FFrameTime SequenceFrameLength = SeqLengthSeconds * FrameRate;
+	const FFrameTime SequenceFrameLength = SeqLengthSeconds * FrameRate;
+	const FFrameTime FirstLoopSequenceFrameLength = FirstLoopSeqLengthInSeconds * FrameRate;
 	if (SequenceFrameLength.FrameNumber > 1)
 	{
 		// Snap to the repeat times
+		bool IsFirstLoop = true;
 		FFrameTime CurrentTime = StartFrame;
 		while (CurrentTime < EndFrame)
 		{
 			OutSnapTimes.Add(CurrentTime.FrameNumber);
-			CurrentTime += SequenceFrameLength;
+			if (IsFirstLoop)
+			{
+				CurrentTime += FirstLoopSequenceFrameLength;
+				IsFirstLoop = false;
+			}
+			else
+			{
+				CurrentTime += SequenceFrameLength;
+			}
 		}
 	}
 }
diff --git a/Engine/Plugins/Experimental/GeometryCache/Source/GeometryCacheTracks/Private/MovieSceneGeometryCacheTemplate.cpp b/Engine/Plugins/Experimental/GeometryCache/Source/GeometryCacheTracks/Private/MovieSceneGeometryCacheTemplate.cpp
--- a/Engine/Plugins/Experimental/GeometryCache/Source/GeometryCacheTracks/Private/MovieSceneGeometryCacheTemplate.cpp
+++ b/Engine/Plugins/Experimental/GeometryCache/Source/GeometryCacheTracks/Private/MovieSceneGeometryCacheTemplate.cpp
@@ -123,20 +123,22 @@ void FMovieSceneGeometryCacheSectionTemplate::Evaluate(const FMovieSceneEvaluati
 
 float FMovieSceneGeometryCacheSectionTemplateParameters::MapTimeToAnimation(float ComponentDuration, FFrameTime InPosition, FFrameRate InFrameRate) const
 {
-	float SequenceLength = ComponentDuration;
-	FFrameTime AnimationLength = SequenceLength * InFrameRate;
-	int32 LengthInFrames = AnimationLength.FrameNumber.Value + (int)(AnimationLength.GetSubFrame() + 0.5f) + 1;
+	const float SequenceLength = ComponentDuration;
+	const FFrameTime AnimationLength = SequenceLength * InFrameRate;
+	const int32 LengthInFrames = AnimationLength.FrameNumber.Value + (int)(AnimationLength.GetSubFrame() + 0.5f) + 1;
 	//we only play end if we are not looping, and assuming we are looping if Length is greater than default length;
-	bool bLooping = (SectionEndTime.Value - SectionStartTime.Value + StartFrameOffset + EndFrameOffset) > LengthInFrames;
+	const bool bLooping = (SectionEndTime.Value - SectionStartTime.Value + StartFrameOffset + EndFrameOffset) > LengthInFrames;
 
 	InPosition = FMath::Clamp(InPosition, FFrameTime(SectionStartTime), FFrameTime(SectionEndTime - 1));
 
 	const float SectionPlayRate = PlayRate;
 	const float AnimPlayRate = FMath::IsNearlyZero(SectionPlayRate) ? 1.0f : SectionPlayRate;
 
+	const float FirstLoopSeqLength = SequenceLength - InFrameRate.AsSeconds(FirstLoopStartFrameOffset + StartFrameOffset + EndFrameOffset);
 	const float SeqLength = SequenceLength - InFrameRate.AsSeconds(StartFrameOffset + EndFrameOffset);
 
 	float AnimPosition = FFrameTime::FromDecimal((InPosition - SectionStartTime).AsDecimal() * AnimPlayRate) / InFrameRate;
+	AnimPosition += InFrameRate.AsSeconds(FirstLoopStartFrameOffset);
 	if (SeqLength > 0.f && (bLooping || !FMath::IsNearlyEqual(AnimPosition, SeqLength, 1e-4f)))
 	{
 		AnimPosition = FMath::Fmod(AnimPosition, SeqLength);
diff --git a/Engine/Source/Editor/MovieSceneTools/Private/TrackEditors/SkeletalAnimationTrackEditor.cpp b/Engine/Source/Editor/MovieSceneTools/Private/TrackEditors/SkeletalAnimationTrackEditor.cpp
--- a/Engine/Source/Editor/MovieSceneTools/Private/TrackEditors/SkeletalAnimationTrackEditor.cpp
+++ b/Engine/Source/Editor/MovieSceneTools/Private/TrackEditors/SkeletalAnimationTrackEditor.cpp
@@ -259,7 +259,7 @@ private:
 FSkeletalAnimationSection::FSkeletalAnimationSection( UMovieSceneSection& InSection, TWeakPtr<ISequencer> InSequencer)
 	: Section(*CastChecked<UMovieSceneSkeletalAnimationSection>(&InSection))
 	, Sequencer(InSequencer)
-	, InitialStartOffsetDuringResize(0)
+	, InitialFirstLoopStartOffsetDuringResize(0)
 	, InitialStartTimeDuringResize(0)
 { }
 
@@ -310,13 +310,14 @@ int32 FSkeletalAnimationSection::OnPaintSection( FSequencerSectionPainter& Paint
 	FFrameRate TickResolution = TimeToPixelConverter.GetTickResolution();
 
 	// Add lines where the animation starts and ends/loops
-	float AnimPlayRate = FMath::IsNearlyZero(Section.Params.PlayRate) ? 1.0f : Section.Params.PlayRate;
-	float SeqLength = (Section.Params.GetSequenceLength() - TickResolution.AsSeconds(Section.Params.StartFrameOffset + Section.Params.EndFrameOffset)) / AnimPlayRate;
+	const float AnimPlayRate = FMath::IsNearlyZero(Section.Params.PlayRate) ? 1.0f : Section.Params.PlayRate;
+	const float SeqLength = (Section.Params.GetSequenceLength() - TickResolution.AsSeconds(Section.Params.StartFrameOffset + Section.Params.EndFrameOffset)) / AnimPlayRate;
+	const float FirstLoopSeqLength = SeqLength - TickResolution.AsSeconds(Section.Params.FirstLoopStartFrameOffset) / AnimPlayRate;
 
 	if (!FMath::IsNearlyZero(SeqLength, KINDA_SMALL_NUMBER) && SeqLength > 0)
 	{
 		float MaxOffset  = Section.GetRange().Size<FFrameTime>() / TickResolution;
-		float OffsetTime = SeqLength;
+		float OffsetTime = FirstLoopSeqLength;
 		float StartTime  = Section.GetInclusiveStartFrame() / TickResolution;
 
 		while (OffsetTime < MaxOffset)
@@ -357,7 +358,7 @@ int32 FSkeletalAnimationSection::OnPaintSection( FSequencerSectionPainter& Paint
 
 void FSkeletalAnimationSection::BeginResizeSection()
 {
-	InitialStartOffsetDuringResize = Section.Params.StartFrameOffset;
+	InitialFirstLoopStartOffsetDuringResize = Section.Params.FirstLoopStartFrameOffset;
 	InitialStartTimeDuringResize   = Section.HasStartFrame() ? Section.GetInclusiveStartFrame() : 0;
 }
 
@@ -369,17 +370,23 @@ void FSkeletalAnimationSection::ResizeSection(ESequencerSectionResizeMode Resize
 		FFrameRate FrameRate   = Section.GetTypedOuter<UMovieScene>()->GetTickResolution();
 		FFrameNumber StartOffset = FrameRate.AsFrameNumber((ResizeTime - InitialStartTimeDuringResize) / FrameRate * Section.Params.PlayRate);
 
-		StartOffset += InitialStartOffsetDuringResize;
+		StartOffset += InitialFirstLoopStartOffsetDuringResize;
 
-		// Ensure start offset is not less than 0 and adjust ResizeTime
 		if (StartOffset < 0)
 		{
+			// Ensure start offset is not less than 0 and adjust ResizeTime
 			ResizeTime = ResizeTime - StartOffset;
 
 			StartOffset = FFrameNumber(0);
 		}
+		else
+		{
+			// If the start offset exceeds the length of one loop, trim it back.
+			const FFrameNumber SeqLength = FrameRate.AsFrameNumber(Section.Params.GetSequenceLength()) - Section.Params.StartFrameOffset - Section.Params.EndFrameOffset;
+			StartOffset = StartOffset % SeqLength;
+		}
 
-		Section.Params.StartFrameOffset = StartOffset;
+		Section.Params.FirstLoopStartFrameOffset = StartOffset;
 	}
 
 	ISequencerSection::ResizeSection(ResizeMode, ResizeTime);
@@ -395,17 +402,23 @@ void FSkeletalAnimationSection::SlipSection(FFrameNumber SlipTime)
 	FFrameRate FrameRate = Section.GetTypedOuter<UMovieScene>()->GetTickResolution();
 	FFrameNumber StartOffset = FrameRate.AsFrameNumber((SlipTime - InitialStartTimeDuringResize) / FrameRate * Section.Params.PlayRate);
 
-	StartOffset += InitialStartOffsetDuringResize;
+	StartOffset += InitialFirstLoopStartOffsetDuringResize;
 
-	// Ensure start offset is not less than 0 and adjust ResizeTime
 	if (StartOffset < 0)
 	{
+		// Ensure start offset is not less than 0 and adjust ResizeTime
 		SlipTime = SlipTime - StartOffset;
 
 		StartOffset = FFrameNumber(0);
 	}
+	else
+	{
+		// If the start offset exceeds the length of one loop, trim it back.
+		const FFrameNumber SeqLength = FrameRate.AsFrameNumber(Section.Params.GetSequenceLength()) - Section.Params.StartFrameOffset - Section.Params.EndFrameOffset;
+		StartOffset = StartOffset % SeqLength;
+	}
 
-	Section.Params.StartFrameOffset = StartOffset;
+	Section.Params.FirstLoopStartFrameOffset = StartOffset;
 
 	ISequencerSection::SlipSection(SlipTime);
 }
diff --git a/Engine/Source/Editor/MovieSceneTools/Private/TrackEditors/SkeletalAnimationTrackEditor.h b/Engine/Source/Editor/MovieSceneTools/Private/TrackEditors/SkeletalAnimationTrackEditor.h
--- a/Engine/Source/Editor/MovieSceneTools/Private/TrackEditors/SkeletalAnimationTrackEditor.h
+++ b/Engine/Source/Editor/MovieSceneTools/Private/TrackEditors/SkeletalAnimationTrackEditor.h
@@ -120,9 +120,9 @@ private:
 	/** Used to draw animation frame, need selection state and local time*/
 	TWeakPtr<ISequencer> Sequencer;
 
-	/** Cached start offset value valid only during resize */
-	FFrameNumber InitialStartOffsetDuringResize;
-	
+	/** Cached first loop start offset value valid only during resize */
+	FFrameNumber InitialFirstLoopStartOffsetDuringResize;
+
 	/** Cached start time valid only during resize */
 	FFrameNumber InitialStartTimeDuringResize;
 };
diff --git a/Engine/Source/Runtime/MovieSceneTracks/Private/Evaluation/MovieSceneSkeletalAnimationTemplate.cpp b/Engine/Source/Runtime/MovieSceneTracks/Private/Evaluation/MovieSceneSkeletalAnimationTemplate.cpp
--- a/Engine/Source/Runtime/MovieSceneTracks/Private/Evaluation/MovieSceneSkeletalAnimationTemplate.cpp
+++ b/Engine/Source/Runtime/MovieSceneTracks/Private/Evaluation/MovieSceneSkeletalAnimationTemplate.cpp
@@ -535,20 +535,21 @@ void FMovieSceneSkeletalAnimationSectionTemplate::Evaluate(const FMovieSceneEval
 
 float FMovieSceneSkeletalAnimationSectionTemplateParameters::MapTimeToAnimation(FFrameTime InPosition, FFrameRate InFrameRate) const
 {
-
-	FFrameTime AnimationLength = GetSequenceLength() * InFrameRate;
-	int32 LengthInFrames = AnimationLength.FrameNumber.Value + (int)(AnimationLength.GetSubFrame() + 0.5f) + 1;
+	const FFrameTime AnimationLength = GetSequenceLength() * InFrameRate;
+	const int32 LengthInFrames = AnimationLength.FrameNumber.Value + (int)(AnimationLength.GetSubFrame() + 0.5f) + 1;
 	//we only play end if we are not looping, and assuming we are looping if Length is greater than default length;
-	bool bLooping = (SectionEndTime.Value - SectionStartTime.Value + StartFrameOffset + EndFrameOffset) > LengthInFrames;
+	const bool bLooping = (SectionEndTime.Value - SectionStartTime.Value + StartFrameOffset + EndFrameOffset) > LengthInFrames;
 
 	InPosition = FMath::Clamp(InPosition, FFrameTime(SectionStartTime), FFrameTime(SectionEndTime -1));
 	
 	const float SectionPlayRate = PlayRate * Animation->RateScale;
 	const float AnimPlayRate = FMath::IsNearlyZero(SectionPlayRate) ? 1.0f : SectionPlayRate;
 
+	const float FirstLoopSeqLength = GetSequenceLength() - InFrameRate.AsSeconds(FirstLoopStartFrameOffset + StartFrameOffset + EndFrameOffset);
 	const float SeqLength = GetSequenceLength() - InFrameRate.AsSeconds(StartFrameOffset + EndFrameOffset);
 
 	float AnimPosition = FFrameTime::FromDecimal((InPosition - SectionStartTime).AsDecimal() * AnimPlayRate) / InFrameRate;
+	AnimPosition += InFrameRate.AsSeconds(FirstLoopStartFrameOffset);
 	if (SeqLength > 0.f && (bLooping || !FMath::IsNearlyEqual(AnimPosition, SeqLength,1e-4f)))
 	{
 		AnimPosition = FMath::Fmod(AnimPosition, SeqLength);
diff --git a/Engine/Source/Runtime/MovieSceneTracks/Private/Sections/MovieSceneSkeletalAnimationSection.cpp b/Engine/Source/Runtime/MovieSceneTracks/Private/Sections/MovieSceneSkeletalAnimationSection.cpp
--- a/Engine/Source/Runtime/MovieSceneTracks/Private/Sections/MovieSceneSkeletalAnimationSection.cpp
+++ b/Engine/Source/Runtime/MovieSceneTracks/Private/Sections/MovieSceneSkeletalAnimationSection.cpp
@@ -64,7 +64,7 @@ UMovieSceneSkeletalAnimationSection::UMovieSceneSkeletalAnimationSection( const
 
 TOptional<FFrameTime> UMovieSceneSkeletalAnimationSection::GetOffsetTime() const
 {
-	return TOptional<FFrameTime>(Params.StartFrameOffset);
+	return TOptional<FFrameTime>(Params.FirstLoopStartFrameOffset);
 }
 
 void UMovieSceneSkeletalAnimationSection::Serialize(FArchive& Ar)
@@ -168,14 +168,18 @@ FMovieSceneEvalTemplatePtr UMovieSceneSkeletalAnimationSection::GenerateTemplate
 	return FMovieSceneSkeletalAnimationSectionTemplate(*this);
 }
 
-FFrameNumber GetStartOffsetAtTrimTime(FQualifiedFrameTime TrimTime, const FMovieSceneSkeletalAnimationParams& Params, FFrameNumber StartFrame, FFrameRate FrameRate)
+FFrameNumber GetFirstLoopStartOffsetAtTrimTime(FQualifiedFrameTime TrimTime, const FMovieSceneSkeletalAnimationParams& Params, FFrameNumber StartFrame, FFrameRate FrameRate)
 {
-	float AnimPlayRate = FMath::IsNearlyZero(Params.PlayRate) ? 1.0f : Params.PlayRate;
-	float AnimPosition = (TrimTime.Time - StartFrame) / TrimTime.Rate * AnimPlayRate;
-	float SeqLength = Params.GetSequenceLength() - FrameRate.AsSeconds(Params.StartFrameOffset + Params.EndFrameOffset) / AnimPlayRate;
+	const float AnimPlayRate = FMath::IsNearlyZero(Params.PlayRate) ? 1.0f : Params.PlayRate;
+	const float AnimPosition = (TrimTime.Time - StartFrame) / TrimTime.Rate * AnimPlayRate;
+	const float SeqLength = Params.GetSequenceLength() - FrameRate.AsSeconds(Params.StartFrameOffset + Params.EndFrameOffset) / AnimPlayRate;
 
 	FFrameNumber NewOffset = FrameRate.AsFrameNumber(FMath::Fmod(AnimPosition, SeqLength));
-	NewOffset += Params.StartFrameOffset;
+	NewOffset += Params.FirstLoopStartFrameOffset;
+
+	const FFrameNumber SeqLengthInFrames = FrameRate.AsFrameNumber(SeqLength);
+	while (NewOffset >= SeqLengthInFrames)
+		NewOffset -= SeqLengthInFrames;
 
 	return NewOffset;
 }
@@ -202,7 +206,7 @@ void UMovieSceneSkeletalAnimationSection::TrimSection(FQualifiedFrameTime TrimTi
 		{
 			FFrameRate FrameRate = GetTypedOuter<UMovieScene>()->GetTickResolution();
 
-			Params.StartFrameOffset = HasStartFrame() ? GetStartOffsetAtTrimTime(TrimTime, Params, GetInclusiveStartFrame(), FrameRate) : 0;
+			Params.FirstLoopStartFrameOffset = HasStartFrame() ? GetFirstLoopStartOffsetAtTrimTime(TrimTime, Params, GetInclusiveStartFrame(), FrameRate) : 0;
 		}
 
 		Super::TrimSection(TrimTime, bTrimLeft);
@@ -213,13 +217,13 @@ UMovieSceneSection* UMovieSceneSkeletalAnimationSection::SplitSection(FQualified
 {
 	FFrameRate FrameRate = GetTypedOuter<UMovieScene>()->GetTickResolution();
 
-	const FFrameNumber NewOffset = HasStartFrame() ? GetStartOffsetAtTrimTime(SplitTime, Params, GetInclusiveStartFrame(), FrameRate) : 0;
+	const FFrameNumber NewOffset = HasStartFrame() ? GetFirstLoopStartOffsetAtTrimTime(SplitTime, Params, GetInclusiveStartFrame(), FrameRate) : 0;
 
 	UMovieSceneSection* NewSection = Super::SplitSection(SplitTime);
 	if (NewSection != nullptr)
 	{
 		UMovieSceneSkeletalAnimationSection* NewSkeletalSection = Cast<UMovieSceneSkeletalAnimationSection>(NewSection);
-		NewSkeletalSection->Params.StartFrameOffset = NewOffset;
+		NewSkeletalSection->Params.FirstLoopStartFrameOffset = NewOffset;
 	}
 	return NewSection;
 }
@@ -235,16 +239,27 @@ void UMovieSceneSkeletalAnimationSection::GetSnapTimes(TArray<FFrameNumber>& Out
 
 	const float AnimPlayRate     = FMath::IsNearlyZero(Params.PlayRate) ? 1.0f : Params.PlayRate;
 	const float SeqLengthSeconds = Params.GetSequenceLength() - FrameRate.AsSeconds(Params.StartFrameOffset + Params.EndFrameOffset) / AnimPlayRate;
+	const float FirstLoopSeqLengthSeconds = SeqLengthSeconds - FrameRate.AsSeconds(Params.FirstLoopStartFrameOffset) / AnimPlayRate;
 
-	FFrameTime SequenceFrameLength = SeqLengthSeconds * FrameRate;
+	const FFrameTime SequenceFrameLength = SeqLengthSeconds * FrameRate;
+	const FFrameTime FirstLoopSequenceFrameLength = FirstLoopSeqLengthSeconds * FrameRate;
 	if (SequenceFrameLength.FrameNumber > 1)
 	{
 		// Snap to the repeat times
+		bool IsFirstLoop = true;
 		FFrameTime CurrentTime = StartFrame;
 		while (CurrentTime < EndFrame)
 		{
 			OutSnapTimes.Add(CurrentTime.FrameNumber);
-			CurrentTime += SequenceFrameLength;
+			if (IsFirstLoop)
+			{
+				CurrentTime += FirstLoopSequenceFrameLength;
+				IsFirstLoop = false;
+			}
+			else
+			{
+				CurrentTime += SequenceFrameLength;
+			}
 		}
 	}
 }
diff --git a/Engine/Source/Runtime/MovieSceneTracks/Public/Sections/MovieSceneSkeletalAnimationSection.h b/Engine/Source/Runtime/MovieSceneTracks/Public/Sections/MovieSceneSkeletalAnimationSection.h
--- a/Engine/Source/Runtime/MovieSceneTracks/Public/Sections/MovieSceneSkeletalAnimationSection.h
+++ b/Engine/Source/Runtime/MovieSceneTracks/Public/Sections/MovieSceneSkeletalAnimationSection.h
@@ -26,6 +26,10 @@ struct FMovieSceneSkeletalAnimationParams
 	UPROPERTY(BlueprintReadWrite, EditAnywhere, Category="Animation", meta=(AllowedClasses = "AnimSequence,AnimComposite,AnimStreamable"))
 	UAnimSequenceBase* Animation;
 
+	/** The offset into the beginning of the animation clip for the first loop of play. */
+	UPROPERTY(BlueprintReadWrite, EditAnywhere, Category="Animation")
+	FFrameNumber FirstLoopStartFrameOffset;
+
 	/** The offset into the beginning of the animation clip */
 	UPROPERTY(BlueprintReadWrite, EditAnywhere, Category="Animation")
 	FFrameNumber StartFrameOffset;
