Fix a static analysis warning that was introduced when the code was updated to use FBulkDataIORequest.

Previously the left hand side of the if statement was just comparing two pointers, but now IORequest is a container around the raw IAsyncReadRequest which means that if IORequest is nullptr (which it will not be) we would crash.
However the right hand side condition was making the left hand side redundant so we could just remove it.
So now the code will will continue to process the first completed chunk that the resident chunk is waiting for and then ignore subsequent completions as intended.

diff --git a/Engine/Plugins/Experimental/GeometryCache/Source/GeometryCache/Private/StreamingGeometryCacheData.cpp b/Engine/Plugins/Experimental/GeometryCache/Source/GeometryCache/Private/StreamingGeometryCacheData.cpp
--- a/Engine/Plugins/Experimental/GeometryCache/Source/GeometryCache/Private/StreamingGeometryCacheData.cpp
+++ b/Engine/Plugins/Experimental/GeometryCache/Source/GeometryCache/Private/StreamingGeometryCacheData.cpp
@@ -328,8 +328,10 @@ void FStreamingGeometryCacheData::ProcessCompletedChunks()
 			return;
 		}
 
-		// Chunks can be queued up multiple times when scrubbing, either the request has already been fulfilled (nullptr) or it could be waiting for a different IO request
-		if (*Chunk->IORequest == CompletedChunk.ReadRequest || Chunk->IORequest != nullptr)
+		// Chunks can be queued up multiple times when scrubbing but we can trust the LoadedChunkIndex of the completed chunk
+		// so all we need to check is if the FResidentChunk  has a valid IORequest pointer or not.
+		// If it is nullptr then we already processed a request for the chunk and it can be ignored.
+		if (Chunk->IORequest != nullptr)
 		{
 			// Check to see if we successfully managed to load anything
 			uint8* Mem = CompletedChunk.ReadRequest->GetReadResults();
