Fix an assert when running the FStreamingGeometryCacheData system via the IoDispatcher.

Instead of calling ::GetFilename and checking the filename length, we should call ::CanLoadFromDisk to see if it is possible to actually load a Bulkdata object as the IoDispatcher means we might not always have a valid file name.

#rnx

diff --git a/Engine/Plugins/Experimental/GeometryCache/Source/GeometryCache/Private/StreamingGeometryCacheData.cpp b/Engine/Plugins/Experimental/GeometryCache/Source/GeometryCache/Private/StreamingGeometryCacheData.cpp
--- a/Engine/Plugins/Experimental/GeometryCache/Source/GeometryCache/Private/StreamingGeometryCacheData.cpp
+++ b/Engine/Plugins/Experimental/GeometryCache/Source/GeometryCache/Private/StreamingGeometryCacheData.cpp
@@ -217,8 +217,8 @@ void FStreamingGeometryCacheData::UpdateStreamingStatus()
 					continue;
 				}
 
-				checkf(Chunk.BulkData.GetFilename().Len(), TEXT("Bulk data is not loaded and not associated with a file."));
-				check(!Chunk.BulkData.IsStoredCompressedOnDisk());
+				checkf(Chunk.BulkData.CanLoadFromDisk(), TEXT("Bulk data is not loaded and cannot be loaded from disk!"));
+				check(!Chunk.BulkData.IsStoredCompressedOnDisk()); // We do not support compressed Bulkdata for this system
 
 				FResidentChunk &ResidentChunk = AddResidentChunk(NeededIndex, Chunk);
 
