Whitelisted GeometryCache plugin on Mac.

diff --git a/Engine/Plugins/Experimental/GeometryCache/GeometryCache.uplugin b/Engine/Plugins/Experimental/GeometryCache/GeometryCache.uplugin
--- a/Engine/Plugins/Experimental/GeometryCache/GeometryCache.uplugin
+++ b/Engine/Plugins/Experimental/GeometryCache/GeometryCache.uplugin
@@ -22,7 +22,8 @@
 			"LoadingPhase" : "Default",
 			"WhitelistPlatforms" :
 			[
-				"Win64"
+				"Win64",
+				"Mac"
 			]
 		},
 		{
@@ -31,7 +32,8 @@
 			"LoadingPhase" : "Default",
 			"WhitelistPlatforms" :
 			[
-				"Win64"
+				"Win64",
+				"Mac"
 			]
 		},
 		{
