Scheduled the GeometryCache prefetch on the game thread as needed.

diff --git a/Engine/Plugins/Experimental/GeometryCache/Source/GeometryCache/Private/GeometryCacheStreamingManager.cpp b/Engine/Plugins/Experimental/GeometryCache/Source/GeometryCache/Private/GeometryCacheStreamingManager.cpp
--- a/Engine/Plugins/Experimental/GeometryCache/Source/GeometryCache/Private/GeometryCacheStreamingManager.cpp
+++ b/Engine/Plugins/Experimental/GeometryCache/Source/GeometryCache/Private/GeometryCacheStreamingManager.cpp
@@ -11,6 +11,7 @@
 #include "GeometryCacheCodecBase.h"
 #include "GeometryCacheModule.h"
 #include "StreamingGeometryCacheData.h"
+#include "Async/Async.h"
 
 DECLARE_CYCLE_STAT(TEXT("Update Resource Streaming"), STAT_UpdateResourceStreaming, STATGROUP_GeometryCache);
 DECLARE_CYCLE_STAT(TEXT("Wait Untill Requests Finished"), STAT_BlockTillAllRequestsFinished, STATGROUP_GeometryCache);
@@ -66,6 +67,8 @@ struct FGeometryCacheStreamingManager : public IGeometryCacheStreamingManager
 
 private:
 
+	void PrefetchDataInternal(UGeometryCacheComponent* CacheComponent);
+
 	/** Geometry caches being managed. */
 	TMap<UGeometryCacheTrackStreamable*, FStreamingGeometryCacheData*> StreamingGeometryCaches;
 
@@ -277,7 +280,7 @@ void FGeometryCacheStreamingManager::RemoveStreamingComponent(UGeometryCacheComp
 	StreamingComponents.Remove(CacheComponent);
 }
 
-void FGeometryCacheStreamingManager::PrefetchData(UGeometryCacheComponent* CacheComponent)
+void FGeometryCacheStreamingManager::PrefetchDataInternal(UGeometryCacheComponent* CacheComponent)
 {
 	check(IsInGameThread());
 	check(IsManagedComponent(CacheComponent));
@@ -296,6 +299,27 @@ void FGeometryCacheStreamingManager::PrefetchData(UGeometryCacheComponent* Cache
 	}
 }
 
+void FGeometryCacheStreamingManager::PrefetchData(UGeometryCacheComponent* CacheComponent)
+{
+	if (IsInGameThread())
+	{
+		PrefetchDataInternal(CacheComponent);
+	}
+	else
+	{
+		// The prefetch doesn't need to be executed right now, so schedule it for the game thread
+		FWeakObjectPtr WeakComponent(CacheComponent);
+		AsyncTask(ENamedThreads::GameThread, [this, WeakComponent]()
+		{
+			UGeometryCacheComponent* CacheComponent = Cast<UGeometryCacheComponent>(WeakComponent.Get());
+			if (CacheComponent)
+			{
+				PrefetchDataInternal(CacheComponent);
+			}
+		});
+	}
+}
+
 bool FGeometryCacheStreamingManager::IsManagedComponent(const UGeometryCacheComponent* CacheComponent) const
 {
 	check(IsInGameThread());
diff --git a/Engine/Plugins/Experimental/GeometryCache/Source/GeometryCache/Private/StreamingGeometryCacheData.cpp b/Engine/Plugins/Experimental/GeometryCache/Source/GeometryCache/Private/StreamingGeometryCacheData.cpp
--- a/Engine/Plugins/Experimental/GeometryCache/Source/GeometryCache/Private/StreamingGeometryCacheData.cpp
+++ b/Engine/Plugins/Experimental/GeometryCache/Source/GeometryCache/Private/StreamingGeometryCacheData.cpp
@@ -316,7 +316,10 @@ bool FStreamingGeometryCacheData::BlockTillAllRequestsFinished(float TimeLimit)
 void FStreamingGeometryCacheData::ProcessCompletedChunks()
 {
 	//Note: This function should only be called from code which owns the CriticalSection
-	check(IsInGameThread() || IsInRenderingThread());
+	if (!IsInGameThread() && !IsInRenderingThread())
+	{
+		return;
+	}
 
 	FCompletedChunk CompletedChunk;
 	while (CompletedChunks.Dequeue(CompletedChunk))
