Fixed variable shadowing.

diff --git a/Engine/Plugins/Experimental/GeometryCache/Source/GeometryCache/Private/GeometryCacheStreamingManager.cpp b/Engine/Plugins/Experimental/GeometryCache/Source/GeometryCache/Private/GeometryCacheStreamingManager.cpp
--- a/Engine/Plugins/Experimental/GeometryCache/Source/GeometryCache/Private/GeometryCacheStreamingManager.cpp
+++ b/Engine/Plugins/Experimental/GeometryCache/Source/GeometryCache/Private/GeometryCacheStreamingManager.cpp
@@ -311,10 +311,10 @@ void FGeometryCacheStreamingManager::PrefetchData(UGeometryCacheComponent* Cache
 		FWeakObjectPtr WeakComponent(CacheComponent);
 		AsyncTask(ENamedThreads::GameThread, [this, WeakComponent]()
 		{
-			UGeometryCacheComponent* CacheComponent = Cast<UGeometryCacheComponent>(WeakComponent.Get());
-			if (CacheComponent)
+			UGeometryCacheComponent* Component = Cast<UGeometryCacheComponent>(WeakComponent.Get());
+			if (Component)
 			{
-				PrefetchDataInternal(CacheComponent);
+				PrefetchDataInternal(Component);
 			}
 		});
 	}
