Geometry Cache: Don't evaluate if the geometry cache isn't registered with the world.

diff --git a/Engine/Plugins/Experimental/GeometryCache/Source/GeometryCacheTracks/Private/MovieSceneGeometryCacheTemplate.cpp b/Engine/Plugins/Experimental/GeometryCache/Source/GeometryCacheTracks/Private/MovieSceneGeometryCacheTemplate.cpp
--- a/Engine/Plugins/Experimental/GeometryCache/Source/GeometryCacheTracks/Private/MovieSceneGeometryCacheTemplate.cpp
+++ b/Engine/Plugins/Experimental/GeometryCache/Source/GeometryCacheTracks/Private/MovieSceneGeometryCacheTemplate.cpp
@@ -85,7 +85,7 @@ struct FGeometryCacheExecutionToken
 				if (UObject* Obj = WeakObj.Get())
 				{
 					UGeometryCacheComponent* GeometryComp = GeometryMeshComponentFromObject(Obj);
-					if (GeometryComp)
+					if (GeometryComp && GeometryComp->IsRegistered())
 					{
 						// Set the GeometryCache on the component only if it's set and valid in the Params
 						if (Params.GeometryCacheAsset && Params.GeometryCacheAsset != GeometryComp->GetGeometryCache())
