Stepping through an important geometry cache does not always show the frame corresponding to the time set in sequencer.
#fix In case of a varying topology mesh, make sure that we still adhere to the interpolation factor when selecting a mesh frame to render. (main issue was was interpolation factor being 1 - epsilon rather than
 1.0)

diff --git a/Engine/Plugins/Experimental/GeometryCache/Source/GeometryCache/Private/GeometryCacheSceneProxy.cpp b/Engine/Plugins/Experimental/GeometryCache/Source/GeometryCache/Private/GeometryCacheSceneProxy.cpp
--- a/Engine/Plugins/Experimental/GeometryCache/Source/GeometryCache/Private/GeometryCacheSceneProxy.cpp
+++ b/Engine/Plugins/Experimental/GeometryCache/Source/GeometryCache/Private/GeometryCacheSceneProxy.cpp
@@ -83,6 +83,7 @@ FGeometryCacheSceneProxy::FGeometryCacheSceneProxy(UGeometryCacheComponent* Comp
 			NewSection->FrameIndex = -1;
 			NewSection->UploadedSampleIndex = -1;
 			NewSection->NextFrameIndex = -1;
+			NewSection->InterpolationFactor = 0.f;
 			NewSection->NextFrameMeshData = nullptr;
 
 			// Allocate verts
@@ -621,6 +622,10 @@ void FGeometryCacheSceneProxy::FrameUpdate() const
 			bool bSeeked = false; // Is this frame a seek and thus the previous rendered frame's data invalid
 			bool bDecoderError = false; // If we have a decoder error we don't interpolate and we don't update the vertex buffers
 										// so essentially we just keep the last valid frame...
+			
+			bool bFrameIndicesChanged = false;
+			bool bDifferentInterpolationFactor = FMath::RoundToInt(InterpolationFactor) != FMath::RoundToInt(TrackProxy->InterpolationFactor);
+			TrackProxy->InterpolationFactor = InterpolationFactor;
 
 			// Compare this against the frames we got and keep some/all/none of them
 			// This will work across frames but also within a frame if the mesh is in several views
@@ -679,6 +684,8 @@ void FGeometryCacheSceneProxy::FrameUpdate() const
 						bDecoderError = true;
 					}
 				}
+
+				bFrameIndicesChanged = true;
 			}
 
 			// Check if we can interpolate between the two frames we have available
@@ -882,29 +889,31 @@ void FGeometryCacheSceneProxy::FrameUpdate() const
 				// We just don't interpolate between frames if we got GPU to burn we could someday render twice and stipple fade between it :-D like with lods
 
 				// Only bother uploading if anything changed or when the we failed to decode anything make sure update the gpu buffers regardless
-				if (bDecodedAnything || bDecoderError)
+				if (bFrameIndicesChanged || bDifferentInterpolationFactor || bDecodedAnything || bDecoderError)
 				{
-					//TrackProxy->IndexBuffer.Update(TrackProxy->MeshData->Indices);
+					const bool bNextFrame = !!FMath::RoundToInt(InterpolationFactor);
+					const uint32 FrameIndexToUse = bNextFrame ? TrackProxy->NextFrameIndex : TrackProxy->FrameIndex;
+					const FGeometryCacheMeshData* MeshDataToUse = bNextFrame ? TrackProxy->NextFrameMeshData : TrackProxy->MeshData;
 
-					const int32 NumVertices = TrackProxy->MeshData->Positions.Num();
+					const int32 NumVertices = MeshDataToUse->Positions.Num();
 
-					if (TrackProxy->MeshData->VertexInfo.bHasTangentX)
-						TrackProxy->TangentXBuffer.Update(TrackProxy->MeshData->TangentsX);
-					if (TrackProxy->MeshData->VertexInfo.bHasTangentZ)
-						TrackProxy->TangentZBuffer.Update(TrackProxy->MeshData->TangentsZ);
+					if (MeshDataToUse->VertexInfo.bHasTangentX)
+						TrackProxy->TangentXBuffer.Update(MeshDataToUse->TangentsX);
+					if (MeshDataToUse->VertexInfo.bHasTangentZ)
+						TrackProxy->TangentZBuffer.Update(MeshDataToUse->TangentsZ);
 
-					if (!TrackProxy->MeshData->VertexInfo.bConstantIndices)
-						TrackProxy->IndexBuffer.Update(TrackProxy->MeshData->Indices);
+					if (!MeshDataToUse->VertexInfo.bConstantIndices)
+						TrackProxy->IndexBuffer.Update(MeshDataToUse->Indices);
 
-					if (TrackProxy->MeshData->VertexInfo.bHasUV0)
-						TrackProxy->TextureCoordinatesBuffer.Update(TrackProxy->MeshData->TextureCoordinates);
+					if (MeshDataToUse->VertexInfo.bHasUV0)
+						TrackProxy->TextureCoordinatesBuffer.Update(MeshDataToUse->TextureCoordinates);
 
-					if (TrackProxy->MeshData->VertexInfo.bHasColor0)
-						TrackProxy->ColorBuffer.Update(TrackProxy->MeshData->Colors);
+					if (MeshDataToUse->VertexInfo.bHasColor0)
+						TrackProxy->ColorBuffer.Update(MeshDataToUse->Colors);
 
-					bool bIsCompatibleWithCachedFrame = TrackProxy->IsTopologyCompatible(
+					const bool bIsCompatibleWithCachedFrame = TrackProxy->IsTopologyCompatible(
 							TrackProxy->PositionBufferFrameIndices[TrackProxy->CurrentPositionBufferIndex % 2],
-						TrackProxy->FrameIndex);
+						FrameIndexToUse);
 
 					if (!bHasMotionVectors)
 					{
@@ -912,11 +921,11 @@ void FGeometryCacheSceneProxy::FrameUpdate() const
 						// with a previous buffer referencing a buffer from another topology
 						if (TrackProxy->CurrentPositionBufferIndex == -1 || !bIsCompatibleWithCachedFrame || bSeeked)
 						{
-							TrackProxy->PositionBuffers[0].Update(TrackProxy->MeshData->Positions);
-							TrackProxy->PositionBuffers[1].Update(TrackProxy->MeshData->Positions);
+							TrackProxy->PositionBuffers[0].Update(MeshDataToUse->Positions);
+							TrackProxy->PositionBuffers[1].Update(MeshDataToUse->Positions);
 							TrackProxy->CurrentPositionBufferIndex = 0;
-							TrackProxy->PositionBufferFrameIndices[0] = TrackProxy->FrameIndex;
-							TrackProxy->PositionBufferFrameIndices[1] = TrackProxy->FrameIndex;
+							TrackProxy->PositionBufferFrameIndices[0] = FrameIndexToUse;
+							TrackProxy->PositionBufferFrameIndices[1] = FrameIndexToUse;
 						}
 						// We still use the previous frame's buffer as a motion blur previous position. As interpolation is switched
 						// off the actual time of this previous frame depends on the geometry cache framerate and playback speed
@@ -925,16 +934,16 @@ void FGeometryCacheSceneProxy::FrameUpdate() const
 						else
 						{
 							TrackProxy->CurrentPositionBufferIndex++;
-							TrackProxy->PositionBuffers[TrackProxy->CurrentPositionBufferIndex % 2].Update(TrackProxy->MeshData->Positions);
-							TrackProxy->PositionBufferFrameIndices[TrackProxy->CurrentPositionBufferIndex % 2] = TrackProxy->FrameIndex;
+							TrackProxy->PositionBuffers[TrackProxy->CurrentPositionBufferIndex % 2].Update(MeshDataToUse->Positions);
+							TrackProxy->PositionBufferFrameIndices[TrackProxy->CurrentPositionBufferIndex % 2] = FrameIndexToUse;
 						}
 					}
 					else
 					{
 						TrackProxy->CurrentPositionBufferIndex = 0;
-						TrackProxy->PositionBuffers[0].Update(TrackProxy->MeshData->Positions);
-						TrackProxy->PositionBuffers[1].Update(TrackProxy->MeshData->MotionVectors);
-						TrackProxy->PositionBufferFrameIndices[0] = TrackProxy->FrameIndex;
+						TrackProxy->PositionBuffers[0].Update(MeshDataToUse->Positions);
+						TrackProxy->PositionBuffers[1].Update(MeshDataToUse->MotionVectors);
+						TrackProxy->PositionBufferFrameIndices[0] = FrameIndexToUse;
 						TrackProxy->PositionBufferFrameIndices[1] = -1;
 						TrackProxy->PositionBufferFrameTimes[0] = Time;
 						TrackProxy->PositionBufferFrameTimes[1] = Time;
@@ -1054,11 +1063,11 @@ const FVisibilitySample& FGeomCacheTrackProxy::GetVisibilitySample(float Time, c
 	return FVisibilitySample::InvisibleSample;
 }
 
-void FGeomCacheTrackProxy::FindSampleIndexesFromTime(float Time, bool bLooping, bool bIsPlayingBackwards, int32 &OutFrameIndex, int32 &OutNextFrameIndex, float &InterpolationFactor)
+void FGeomCacheTrackProxy::FindSampleIndexesFromTime(float Time, bool bLooping, bool bIsPlayingBackwards, int32 &OutFrameIndex, int32 &OutNextFrameIndex, float& InInterpolationFactor)
 {
 	if (UGeometryCacheTrackStreamable* StreamableTrack = Cast<UGeometryCacheTrackStreamable>(Track))
 	{
-		StreamableTrack->FindSampleIndexesFromTime(Time, bLooping, bIsPlayingBackwards, OutFrameIndex, OutNextFrameIndex, InterpolationFactor);
+		StreamableTrack->FindSampleIndexesFromTime(Time, bLooping, bIsPlayingBackwards, OutFrameIndex, OutNextFrameIndex, InInterpolationFactor);
 	}
 }
 
diff --git a/Engine/Plugins/Experimental/GeometryCache/Source/GeometryCache/Public/GeometryCacheSceneProxy.h b/Engine/Plugins/Experimental/GeometryCache/Source/GeometryCache/Public/GeometryCacheSceneProxy.h
--- a/Engine/Plugins/Experimental/GeometryCache/Source/GeometryCache/Public/GeometryCacheSceneProxy.h
+++ b/Engine/Plugins/Experimental/GeometryCache/Source/GeometryCache/Public/GeometryCacheSceneProxy.h
@@ -192,7 +192,7 @@ public:
 	 * @param OutNextFrameIndex - The frame index that follows OutFrameIndex
 	 * @param InterpolationFactor - The interpolation value between the two frame times
 	 */
-	virtual void FindSampleIndexesFromTime(float Time, bool bLooping, bool bIsPlayingBackwards, int32 &OutFrameIndex, int32 &OutNextFrameIndex, float &InterpolationFactor);
+	virtual void FindSampleIndexesFromTime(float Time, bool bLooping, bool bIsPlayingBackwards, int32 &OutFrameIndex, int32 &OutNextFrameIndex, float &InInterpolationFactor);
 
 	/** MeshData storing information used for rendering this Track */
 	FGeometryCacheMeshData *MeshData;
@@ -201,6 +201,7 @@ public:
 	/** Frame numbers corresponding to MeshData, NextFrameMeshData */
 	int32 FrameIndex;
 	int32 NextFrameIndex;
+	float InterpolationFactor;
 
 	/** Material applied to this Track */
 	TArray<UMaterialInterface*> Materials;
