Integrated GeometryCache improvements from Reverb (CL12505314). Persistent memory allocations and SIMD frame interpolation.

diff --git a/Engine/Plugins/Experimental/GeometryCache/Source/GeometryCache/Private/GeometryCacheSceneProxy.cpp b/Engine/Plugins/Experimental/GeometryCache/Source/GeometryCache/Private/GeometryCacheSceneProxy.cpp
--- a/Engine/Plugins/Experimental/GeometryCache/Source/GeometryCache/Private/GeometryCacheSceneProxy.cpp
+++ b/Engine/Plugins/Experimental/GeometryCache/Source/GeometryCache/Private/GeometryCacheSceneProxy.cpp
@@ -35,7 +35,6 @@ static TAutoConsoleVariable<int32> CVarInterpolateFrames(
 	TEXT("Interpolate between geometry cache frames (if topology allows this)."),
 	ECVF_Scalability | ECVF_RenderThreadSafe);
 
-
 /**
 * All vertex information except the position.
 */
@@ -205,9 +204,9 @@ FGeometryCacheSceneProxy::~FGeometryCacheSceneProxy()
 			Section->VertexFactory.ReleaseResource();
 			Section->PositionBuffers[0].ReleaseResource();
 			Section->PositionBuffers[1].ReleaseResource();
-		#if RHI_RAYTRACING
+#if RHI_RAYTRACING
 			Section->RayTracingGeometry.ReleaseResource();
-		#endif
+#endif
 			delete Section->MeshData;
 			if (Section->NextFrameMeshData != nullptr)
 				delete Section->NextFrameMeshData;
@@ -289,7 +288,7 @@ public:
 static float OneOver255 = 1.0f / 255.0f;
 
 // Avoid converting from 8 bit normalized to float and back again.
-inline FPackedNormal InterpolatePackedNormal(const FPackedNormal &A, const FPackedNormal &B, int32 ScaledFactor, int32 OneMinusScaledFactor)
+inline FPackedNormal InterpolatePackedNormal(const FPackedNormal& A, const FPackedNormal& B, int32 ScaledFactor, int32 OneMinusScaledFactor)
 {
 	FPackedNormal result;
 	result.Vector.X = (A.Vector.X * OneMinusScaledFactor + B.Vector.X * ScaledFactor) * OneOver255;
@@ -300,7 +299,7 @@ inline FPackedNormal InterpolatePackedNormal(const FPackedNormal &A, const FPack
 }
 
 // Avoid converting from 8 bit normalized to float and back again.
-inline FColor InterpolatePackedColor(const FColor &A, const FColor &B, int32 ScaledFactor, int32 OneMinusScaledFactor)
+inline FColor InterpolatePackedColor(const FColor& A, const FColor& B, int32 ScaledFactor, int32 OneMinusScaledFactor)
 {
 	FColor result;
 	result.R = (A.R * OneMinusScaledFactor + B.R * ScaledFactor) * OneOver255;
@@ -317,8 +316,8 @@ SIZE_T FGeometryCacheSceneProxy::GetTypeHash() const
 }
 
 void FGeometryCacheSceneProxy::CreateMeshBatch(
-	const FGeomCacheTrackProxy* TrackProxy, 
-	const FGeometryCacheMeshBatchInfo& BatchInfo, 
+	const FGeomCacheTrackProxy* TrackProxy,
+	const FGeometryCacheMeshBatchInfo& BatchInfo,
 	FGeometryCacheVertexFactoryUserDataWrapper& UserDataWrapper,
 	FDynamicPrimitiveUniformBuffer& DynamicPrimitiveUniformBuffer,
 	FMeshBatch& Mesh) const
@@ -334,19 +333,19 @@ void FGeometryCacheSceneProxy::CreateMeshBatch(
 		TrackProxy->MeshData->Positions.Num() == TrackProxy->MeshData->MotionVectors.Num())
 		&& (TrackProxy->NextFrameMeshData->Positions.Num() == TrackProxy->NextFrameMeshData->MotionVectors.Num());
 
-				if (!bHasMotionVectors)
-	            {
-		            const float PreviousPositionScale = (GFrameNumber <= UpdatedFrameNum) ? 1.f : 0.f;
-					UserData.MotionBlurDataExtension = FVector::OneVector * PreviousPositionScale;
-					UserData.MotionBlurDataOrigin = FVector::ZeroVector;
-					UserData.MotionBlurPositionScale = 1.f - PreviousPositionScale;
-	            }
-	            else
-	            {
-		            UserData.MotionBlurDataExtension = FVector::OneVector * PlaybackSpeed;
-		            UserData.MotionBlurDataOrigin = FVector::ZeroVector;
-		            UserData.MotionBlurPositionScale = 1.0f;
-	            }
+	if (!bHasMotionVectors)
+	{
+		const float PreviousPositionScale = (GFrameNumber <= UpdatedFrameNum) ? 1.f : 0.f;
+		UserData.MotionBlurDataExtension = FVector::OneVector * PreviousPositionScale;
+		UserData.MotionBlurDataOrigin = FVector::ZeroVector;
+		UserData.MotionBlurPositionScale = 1.f - PreviousPositionScale;
+	}
+	else
+	{
+		UserData.MotionBlurDataExtension = FVector::OneVector * PlaybackSpeed;
+		UserData.MotionBlurDataOrigin = FVector::ZeroVector;
+		UserData.MotionBlurPositionScale = 1.0f;
+	}
 
 	if (IsRayTracingEnabled())
 	{
@@ -357,7 +356,7 @@ void FGeometryCacheSceneProxy::CreateMeshBatch(
 	}
 
 	UserData.PositionBuffer = &TrackProxy->PositionBuffers[TrackProxy->CurrentPositionBufferIndex % 2];
-	UserData.MotionBlurDataBuffer = &TrackProxy->PositionBuffers[(TrackProxy->CurrentPositionBufferIndex+1) % 2];
+	UserData.MotionBlurDataBuffer = &TrackProxy->PositionBuffers[(TrackProxy->CurrentPositionBufferIndex + 1) % 2];
 
 	FGeometryCacheVertexFactoryUniformBufferParameters UniformBufferParameters;
 
@@ -476,10 +475,10 @@ void FGeometryCacheSceneProxy::GetDynamicMeshElements(const TArray<const FSceneV
 						INC_DWORD_STAT_BY(STAT_GeometryCacheSceneProxy_TriangleCount, MeshBatch.Elements[0].NumPrimitives);
 						INC_DWORD_STAT_BY(STAT_GeometryCacheSceneProxy_MeshBatchCount, 1);
 
-					#if !(UE_BUILD_SHIPPING || UE_BUILD_TEST)
+#if !(UE_BUILD_SHIPPING || UE_BUILD_TEST)
 						// Render bounds
 						RenderBounds(Collector.GetPDI(ViewIndex), ViewFamily.EngineShowFlags, GetBounds(), IsSelected());
-					#endif
+#endif
 					}
 				}
 			}
@@ -569,7 +568,7 @@ void FGeometryCacheSceneProxy::UpdateAnimation(float NewTime, bool bNewLooping,
 		// When ray tracing is enabled, update regardless of visibility
 		FrameUpdate();
 
-	#if RHI_RAYTRACING
+#if RHI_RAYTRACING
 		for (FGeomCacheTrackProxy* Section : Tracks)
 		{
 			if (Section != nullptr)
@@ -578,7 +577,7 @@ void FGeometryCacheSceneProxy::UpdateAnimation(float NewTime, bool bNewLooping,
 
 				Section->RayTracingGeometry.Initializer.IndexBuffer = Section->IndexBuffer.IndexBufferRHI;
 				Section->RayTracingGeometry.Initializer.TotalPrimitiveCount = 0;
-				
+
 				TMemoryImageArray<FRayTracingGeometrySegment>& Segments = Section->RayTracingGeometry.Initializer.Segments;
 				Segments.Reset();
 
@@ -592,11 +591,11 @@ void FGeometryCacheSceneProxy::UpdateAnimation(float NewTime, bool bNewLooping,
 					Segments.Add(Segment);
 					Section->RayTracingGeometry.Initializer.TotalPrimitiveCount += BatchInfo.NumTriangles;
 				}
-							
+
 				Section->RayTracingGeometry.UpdateRHI();
 			}
 		}
-	#endif
+#endif
 	}
 }
 
@@ -622,9 +621,10 @@ void FGeometryCacheSceneProxy::FrameUpdate() const
 			bool bSeeked = false; // Is this frame a seek and thus the previous rendered frame's data invalid
 			bool bDecoderError = false; // If we have a decoder error we don't interpolate and we don't update the vertex buffers
 										// so essentially we just keep the last valid frame...
-			
+
 			bool bFrameIndicesChanged = false;
-			bool bDifferentInterpolationFactor = FMath::RoundToInt(InterpolationFactor) != FMath::RoundToInt(TrackProxy->InterpolationFactor);
+			const bool bDifferentRoundedInterpolationFactor = FMath::RoundToInt(InterpolationFactor) != FMath::RoundToInt(TrackProxy->InterpolationFactor);
+			const bool bDifferentInterpolationFactor = !FMath::IsNearlyEqual(InterpolationFactor, TrackProxy->InterpolationFactor);
 			TrackProxy->InterpolationFactor = InterpolationFactor;
 
 			// Compare this against the frames we got and keep some/all/none of them
@@ -635,7 +635,7 @@ void FGeometryCacheSceneProxy::FrameUpdate() const
 				if (TrackProxy->NextFrameIndex == FrameIndex)
 				{
 					// Cycle the current and next frame double buffer
-					FGeometryCacheMeshData *OldFrameMesh = TrackProxy->MeshData;
+					FGeometryCacheMeshData* OldFrameMesh = TrackProxy->MeshData;
 					TrackProxy->MeshData = TrackProxy->NextFrameMeshData;
 					TrackProxy->NextFrameMeshData = OldFrameMesh;
 
@@ -699,7 +699,7 @@ void FGeometryCacheSceneProxy::FrameUpdate() const
 				&& (TrackProxy->NextFrameMeshData->Positions.Num() == TrackProxy->NextFrameMeshData->MotionVectors.Num());
 
 			// Can we interpolate the vertex data?
-			if (bCanInterpolate && !bDecoderError && CVarInterpolateFrames.GetValueOnRenderThread() != 0)
+			if (bCanInterpolate && (bDifferentInterpolationFactor || bFrameIndicesChanged) && !bDecoderError && CVarInterpolateFrames.GetValueOnRenderThread() != 0)
 			{
 				SCOPE_CYCLE_COUNTER(STAT_InterpolateFrames);
 				// Interpolate if the time has changed.
@@ -710,186 +710,299 @@ void FGeometryCacheSceneProxy::FrameUpdate() const
 				// If motion blur suddenly seems to stop working while it should be working it may be that the
 				// CurrentPositionBufferIndex gets inadvertently bumped twice per frame essentially using the same
 				// data for current and previous during rendering.
-				if (TrackProxy->PositionBufferFrameTimes[TrackProxy->CurrentPositionBufferIndex % 2] != Time)
-				{
-					TArray<FVector> InterpolatedPositions;
-					TArray<FPackedNormal> InterpolatedTangentX;
-					TArray<FPackedNormal> InterpolatedTangentZ;
-					TArray<FVector2D> InterpolatedUVs;
-					TArray<FColor> InterpolatedColors;
-
-					const int32 NumVerts = TrackProxy->MeshData->Positions.Num();
-
-					InterpolatedPositions.AddUninitialized(NumVerts);
-					InterpolatedTangentX.AddUninitialized(NumVerts);
-					InterpolatedTangentZ.AddUninitialized(NumVerts);
-					InterpolatedUVs.AddUninitialized(NumVerts);
-					InterpolatedColors.AddUninitialized(NumVerts);
-
-					TArray<FVector> InterpolatedMotionVectors;
-					if (bHasMotionVectors)
-					{
-						InterpolatedMotionVectors.AddUninitialized(NumVerts);
-					}
 
+				const int32 NumVerts = TrackProxy->MeshData->Positions.Num();
+				Scratch.Prepare(NumVerts, bHasMotionVectors);
 
-					const float OneMinusInterp = 1.0 - InterpolationFactor;
-					const int32 InterpFixed = (int32)(InterpolationFactor * 255.0f);
-					const int32 OneMinusInterpFixed = 255 - InterpFixed;
+				const float OneMinusInterp = 1.0 - InterpolationFactor;
+				const int32 InterpFixed = (int32)(InterpolationFactor * 255.0f);
+				const int32 OneMinusInterpFixed = 255 - InterpFixed;
+				const VectorRegister WeightA = VectorSetFloat1( OneMinusInterp );
+				const VectorRegister WeightB = VectorSetFloat1( InterpolationFactor );
+				const VectorRegister Half = VectorSetFloat1( 0.5f );
 
+				#define VALIDATE 0
+				{
+					check(TrackProxy->MeshData->Positions.Num() >= NumVerts);
+					check(TrackProxy->NextFrameMeshData->Positions.Num() >= NumVerts);
+					check(Scratch.InterpolatedPositions.Num() >= NumVerts);
+					const FVector* PositionAPtr = TrackProxy->MeshData->Positions.GetData();
+					const FVector* PositionBPtr = TrackProxy->NextFrameMeshData->Positions.GetData();
+					FVector* InterpolatedPositionsPtr = Scratch.InterpolatedPositions.GetData();
+
+					// Unroll 4 times so we can do 4 wide SIMD
 					{
-						check(TrackProxy->MeshData->Positions.Num() >= NumVerts);
-						check(TrackProxy->NextFrameMeshData->Positions.Num() >= NumVerts);
-						check(InterpolatedPositions.Num() >= NumVerts);
-						const FVector* PositionAPtr = TrackProxy->MeshData->Positions.GetData();
-						const FVector* PositionBPtr = TrackProxy->NextFrameMeshData->Positions.GetData();
-						FVector* InterpolationPositionsPtr = InterpolatedPositions.GetData();
-						for (int32 Index = 0; Index < NumVerts; ++Index)
+						const FVector4* PositionAPtr4 = (const FVector4*)PositionAPtr;
+						const FVector4* PositionBPtr4 = (const FVector4*)PositionBPtr;
+						FVector4* InterpolatedPositionsPtr4 = (FVector4*)InterpolatedPositionsPtr;
+
+						int32 Index = 0;
+						for (; Index + 3 < NumVerts; Index += 4)
 						{
-							const FVector& PositionA = PositionAPtr[Index];
-							const FVector& PositionB = PositionBPtr[Index];
-							InterpolationPositionsPtr[Index] = PositionA * OneMinusInterp + PositionB * InterpolationFactor;
+							VectorRegister Pos0xyz_Pos1x = VectorMultiplyAdd(VectorLoad(PositionAPtr4 + 0), WeightA, VectorMultiply(VectorLoad(PositionBPtr4 + 0), WeightB));
+							VectorRegister Pos1yz_Pos2xy = VectorMultiplyAdd(VectorLoad(PositionAPtr4 + 1), WeightA, VectorMultiply(VectorLoad(PositionBPtr4 + 1), WeightB));
+							VectorRegister Pos2z_Pos3xyz = VectorMultiplyAdd(VectorLoad(PositionAPtr4 + 2), WeightA, VectorMultiply(VectorLoad(PositionBPtr4 + 2), WeightB));
+							VectorStore(Pos0xyz_Pos1x, InterpolatedPositionsPtr4 + 0);
+							VectorStore(Pos1yz_Pos2xy, InterpolatedPositionsPtr4 + 1);
+							VectorStore(Pos2z_Pos3xyz, InterpolatedPositionsPtr4 + 2);
+							PositionAPtr4 += 3;
+							PositionBPtr4 += 3;
+							InterpolatedPositionsPtr4 += 3;
 						}
-					}
-					
-					{
-						check(TrackProxy->MeshData->TangentsX.Num() >= NumVerts);
-						check(TrackProxy->NextFrameMeshData->TangentsX.Num() >= NumVerts);
-						check(TrackProxy->MeshData->TangentsZ.Num() >= NumVerts);
-						check(TrackProxy->NextFrameMeshData->TangentsZ.Num() >= NumVerts);
-						check(InterpolatedTangentX.Num() >= NumVerts);
-						check(InterpolatedTangentZ.Num() >= NumVerts);
-						const FPackedNormal* TangentXAPtr = TrackProxy->MeshData->TangentsX.GetData();
-						const FPackedNormal* TangentXBPtr = TrackProxy->NextFrameMeshData->TangentsX.GetData();
-						const FPackedNormal* TangentZAPtr = TrackProxy->MeshData->TangentsZ.GetData();
-						const FPackedNormal* TangentZBPtr = TrackProxy->NextFrameMeshData->TangentsZ.GetData();
-						FPackedNormal* InterpolatedTangentXPtr = InterpolatedTangentX.GetData();
-						FPackedNormal* InterpolatedTangentZPtr = InterpolatedTangentZ.GetData();
-						for (int32 Index = 0; Index < NumVerts; ++Index)
+
+						for (; Index < NumVerts; Index++)
 						{
-							// The following are already 8 bit so quantized enough we can do exact equal comparisons
-							const FPackedNormal& TangentXA = TangentXAPtr[Index];
-							const FPackedNormal& TangentXB = TangentXBPtr[Index];
-							const FPackedNormal& TangentZA = TangentZAPtr[Index];
-							const FPackedNormal& TangentZB = TangentZBPtr[Index];
-
-							InterpolatedTangentXPtr[Index] = InterpolatePackedNormal(TangentXA, TangentXB, InterpFixed, OneMinusInterpFixed);
-							InterpolatedTangentZPtr[Index] = InterpolatePackedNormal(TangentZA, TangentZB, InterpFixed, OneMinusInterpFixed);
+							InterpolatedPositionsPtr[Index] = PositionAPtr[Index] * OneMinusInterp + PositionBPtr[Index] * InterpolationFactor;
 						}
 					}
-					
+#if VALIDATE
+					for (int32 Index = 0; Index < NumVerts; ++Index)
+					{
+						FVector Result = PositionAPtr[Index] * OneMinusInterp + PositionBPtr[Index] * InterpolationFactor;
+						check(FMath::Abs(InterpolatedPositionsPtr[Index].X - Result.X) < 0.01f);
+						check(FMath::Abs(InterpolatedPositionsPtr[Index].Y - Result.Y) < 0.01f);
+						check(FMath::Abs(InterpolatedPositionsPtr[Index].Z - Result.Z) < 0.01f);
+					}
+#endif
+				}
+				
+				{
+					check(TrackProxy->MeshData->TangentsX.Num() >= NumVerts);
+					check(TrackProxy->NextFrameMeshData->TangentsX.Num() >= NumVerts);
+					check(TrackProxy->MeshData->TangentsZ.Num() >= NumVerts);
+					check(TrackProxy->NextFrameMeshData->TangentsZ.Num() >= NumVerts);
+					check(Scratch.InterpolatedTangentX.Num() >= NumVerts);
+					check(Scratch.InterpolatedTangentZ.Num() >= NumVerts);
+					const FPackedNormal* TangentXAPtr = TrackProxy->MeshData->TangentsX.GetData();
+					const FPackedNormal* TangentXBPtr = TrackProxy->NextFrameMeshData->TangentsX.GetData();
+					const FPackedNormal* TangentZAPtr = TrackProxy->MeshData->TangentsZ.GetData();
+					const FPackedNormal* TangentZBPtr = TrackProxy->NextFrameMeshData->TangentsZ.GetData();
+					FPackedNormal* InterpolatedTangentXPtr = Scratch.InterpolatedTangentX.GetData();
+					FPackedNormal* InterpolatedTangentZPtr = Scratch.InterpolatedTangentZ.GetData();
+
+					const uint32 SignMask = 0x80808080u;
+					for (int32 Index = 0; Index < NumVerts; ++Index)
+					{
+						// VectorLoadSignedByte4 on all inputs is significantly more expensive than VectorLoadByte4, so lets just use unsigned.
+						// Interpolating signed values as unsigned is not correct, but if we flip the signs first it is!
+						// Flipping the sign maps the signed range [-128, 127] to the unsigned range [0, 255]
+						// Unsigned value with flip			Signed value
+						// 0								-128
+						// 1								-127
+						// ..								..
+						// 127								-1
+						// 128								0
+						// 129								1
+						// 255								127
+
+						uint32 TangentXA = TangentXAPtr[Index].Vector.Packed ^ SignMask;
+						uint32 TangentXB = TangentXBPtr[Index].Vector.Packed ^ SignMask;
+						VectorRegister InterpolatedTangentX =	VectorMultiplyAdd(	VectorLoadByte4(&TangentXA), WeightA, 
+																VectorMultiplyAdd(	VectorLoadByte4(&TangentXB), WeightB, Half));	// +0.5f so truncation becomes round to nearest.
+						uint32 PackedInterpolatedTangentX;
+						VectorStoreByte4(InterpolatedTangentX, &PackedInterpolatedTangentX);
+						InterpolatedTangentXPtr[Index].Vector.Packed = PackedInterpolatedTangentX ^ SignMask;	// Convert back to signed
+
+						uint32 TangentZA = TangentZAPtr[Index].Vector.Packed ^ SignMask;
+						uint32 TangentZB = TangentZBPtr[Index].Vector.Packed ^ SignMask;
+						VectorRegister InterpolatedTangentZ =	VectorMultiplyAdd(	VectorLoadByte4(&TangentZA), WeightA, 
+																VectorMultiplyAdd(	VectorLoadByte4(&TangentZB), WeightB, Half));	// +0.5f so truncation becomes round to nearest.
+						uint32 PackedInterpolatedTangentZ;
+						VectorStoreByte4(InterpolatedTangentZ, &PackedInterpolatedTangentZ);
+						InterpolatedTangentZPtr[Index].Vector.Packed = PackedInterpolatedTangentZ ^ SignMask;	// Convert back to signed
+					}
+					VectorResetFloatRegisters();	//TODO: is this actually needed on any platform?
 
-					if (TrackProxy->MeshData->VertexInfo.bHasColor0)
+#if VALIDATE
+					for (int32 Index = 0; Index < NumVerts; ++Index)
 					{
-						check(TrackProxy->MeshData->Colors.Num() >= NumVerts);
-						check(TrackProxy->NextFrameMeshData->Colors.Num() >= NumVerts);
-						check(InterpolatedColors.Num() >= NumVerts);
-						const FColor* ColorAPtr = TrackProxy->MeshData->Colors.GetData();
-						const FColor* ColorBPtr = TrackProxy->NextFrameMeshData->Colors.GetData();
-						FColor* InterpolatedColorsPtr = InterpolatedColors.GetData();
-						for (int32 Index = 0; Index < NumVerts; ++Index)
-						{
-							const FColor& ColorA = ColorAPtr[Index];
-							const FColor& ColorB = ColorBPtr[Index];
-							InterpolatedColorsPtr[Index] = InterpolatePackedColor(ColorA, ColorB, InterpFixed, OneMinusInterpFixed);
-						}
+						FPackedNormal ResultX = InterpolatePackedNormal(TangentXAPtr[Index], TangentXBPtr[Index], InterpFixed, OneMinusInterpFixed);
+						FPackedNormal ResultZ = InterpolatePackedNormal(TangentZAPtr[Index], TangentZBPtr[Index], InterpFixed, OneMinusInterpFixed);
+						check(FMath::Abs(InterpolatedTangentXPtr[Index].Vector.X - ResultX.Vector.X) <= 2);
+						check(FMath::Abs(InterpolatedTangentXPtr[Index].Vector.Y - ResultX.Vector.Y) <= 2);
+						check(FMath::Abs(InterpolatedTangentXPtr[Index].Vector.Z - ResultX.Vector.Z) <= 2);
+						check(FMath::Abs(InterpolatedTangentXPtr[Index].Vector.W - ResultX.Vector.W) <= 2);
+
+						check(FMath::Abs(InterpolatedTangentZPtr[Index].Vector.X - ResultZ.Vector.X) <= 2);
+						check(FMath::Abs(InterpolatedTangentZPtr[Index].Vector.Y - ResultZ.Vector.Y) <= 2);
+						check(FMath::Abs(InterpolatedTangentZPtr[Index].Vector.Z - ResultZ.Vector.Z) <= 2);
+						check(FMath::Abs(InterpolatedTangentZPtr[Index].Vector.W - ResultZ.Vector.W) <= 2);
 					}
+#endif
+				}
 
-					if (TrackProxy->MeshData->VertexInfo.bHasUV0)
+				if (TrackProxy->MeshData->VertexInfo.bHasColor0)
+				{
+					check(TrackProxy->MeshData->Colors.Num() >= NumVerts);
+					check(TrackProxy->NextFrameMeshData->Colors.Num() >= NumVerts);
+					check(Scratch.InterpolatedColors.Num() >= NumVerts);
+					const FColor* ColorAPtr = TrackProxy->MeshData->Colors.GetData();
+					const FColor* ColorBPtr = TrackProxy->NextFrameMeshData->Colors.GetData();
+					FColor* InterpolatedColorsPtr = Scratch.InterpolatedColors.GetData();
+
+					for( int32 Index = 0; Index < NumVerts; ++Index )
+					{
+						VectorRegister InterpolatedColor =		VectorMultiplyAdd( VectorLoadByte4( &ColorAPtr[Index] ), WeightA,
+																VectorMultiplyAdd( VectorLoadByte4( &ColorBPtr[Index] ), WeightB, Half ) );	// +0.5f so truncation becomes round to nearest.
+						VectorStoreByte4(InterpolatedColor, &InterpolatedColorsPtr[Index]);
+					}
+#if VALIDATE
+					for(int32 Index = 0; Index < NumVerts; ++Index)
 					{
-						check(TrackProxy->MeshData->TextureCoordinates.Num() >= NumVerts);
-						check(TrackProxy->NextFrameMeshData->TextureCoordinates.Num() >= NumVerts);
-						check(InterpolatedUVs.Num() >= NumVerts);
-						const FVector2D* UVAPtr = TrackProxy->MeshData->TextureCoordinates.GetData();
-						const FVector2D* UVBPtr = TrackProxy->NextFrameMeshData->TextureCoordinates.GetData();
-						FVector2D* InterpolatedUVsPtr = InterpolatedUVs.GetData();
-						for (int32 Index = 0; Index < NumVerts; ++Index)
+						const FColor& ColorA = ColorAPtr[Index];
+						const FColor& ColorB = ColorBPtr[Index];
+						FColor Result = InterpolatePackedColor(ColorA, ColorB, InterpFixed, OneMinusInterpFixed);
+						check(FMath::Abs(InterpolatedColorsPtr[Index].R - Result.R) <= 1);
+						check(FMath::Abs(InterpolatedColorsPtr[Index].G - Result.G) <= 1);
+						check(FMath::Abs(InterpolatedColorsPtr[Index].B - Result.B) <= 1);
+						check(FMath::Abs(InterpolatedColorsPtr[Index].A - Result.A) <= 1);
+					}
+#endif
+				}
+
+				if (TrackProxy->MeshData->VertexInfo.bHasUV0)
+				{
+					check(TrackProxy->MeshData->TextureCoordinates.Num() >= NumVerts);
+					check(TrackProxy->NextFrameMeshData->TextureCoordinates.Num() >= NumVerts);
+					check(Scratch.InterpolatedUVs.Num() >= NumVerts);
+					const FVector2D* UVAPtr = TrackProxy->MeshData->TextureCoordinates.GetData();
+					const FVector2D* UVBPtr = TrackProxy->NextFrameMeshData->TextureCoordinates.GetData();
+					FVector2D* InterpolatedUVsPtr = Scratch.InterpolatedUVs.GetData();
+
+					// Unroll 2x so we can use 4 wide ops. OOP will hopefully take care of the rest.
+					{
+						int32 Index = 0;
+						for (; Index + 1 < NumVerts; Index += 2)
 						{
-							const FVector2D& UVA = UVAPtr[Index];
-							const FVector2D& UVB = UVBPtr[Index];
-							InterpolatedUVsPtr[Index] = UVA * OneMinusInterp + UVB * InterpolationFactor;
+							VectorRegister InterpolatedUVx2 = VectorMultiplyAdd(	VectorLoad(&UVAPtr[Index]), WeightA,
+																VectorMultiply(		VectorLoad(&UVBPtr[Index]), WeightB));
+							VectorStore(InterpolatedUVx2, &InterpolatedUVsPtr[Index]);
 						}
+
+						if(Index < NumVerts)
+						{
+							InterpolatedUVsPtr[Index] = UVAPtr[Index] * OneMinusInterp + UVBPtr[Index] * InterpolationFactor;
+						}
+					}
+						
+#if VALIDATE
+					for (int32 Index = 0; Index < NumVerts; ++Index)
+					{
+						FVector2D Result = UVAPtr[Index] * OneMinusInterp + UVBPtr[Index] * InterpolationFactor;
+						check(FMath::Abs(InterpolatedUVsPtr[Index].X - Result.X) < 0.01f);
+						check(FMath::Abs(InterpolatedUVsPtr[Index].Y - Result.Y) < 0.01f);
 					}
+#endif
+				}
 
-					if (bHasMotionVectors)
+				if (bHasMotionVectors)
+				{
+					check(TrackProxy->MeshData->MotionVectors.Num() >= NumVerts);
+					check(TrackProxy->NextFrameMeshData->MotionVectors.Num() >= NumVerts);
+					check(Scratch.InterpolatedMotionVectors.Num() >= NumVerts);
+					const FVector* MotionVectorsAPtr = TrackProxy->MeshData->MotionVectors.GetData();
+					const FVector* MotionVectorsBPtr = TrackProxy->NextFrameMeshData->MotionVectors.GetData();
+					FVector* InterpolatedMotionVectorsPtr = Scratch.InterpolatedMotionVectors.GetData();
+
+					// Unroll 4 times so we can do 4 wide SIMD
 					{
-						check(TrackProxy->MeshData->MotionVectors.Num() >= NumVerts);
-						check(TrackProxy->NextFrameMeshData->MotionVectors.Num() >= NumVerts);
-						check(InterpolatedMotionVectors.Num() >= NumVerts);
-						const FVector* MotionVectorAPtr = TrackProxy->MeshData->MotionVectors.GetData();
-						const FVector* MotionVectorBPtr = TrackProxy->NextFrameMeshData->MotionVectors.GetData();
-						FVector* InterpolatedMotionVectorsPtr = InterpolatedMotionVectors.GetData();
-
-						for (int32 Index = 0; Index < NumVerts; ++Index)
+						const FVector4* MotionVectorsAPtr4 = (const FVector4*)MotionVectorsAPtr;
+						const FVector4* MotionVectorsBPtr4 = (const FVector4*)MotionVectorsBPtr;
+						FVector4* InterpolatedMotionVectorsPtr4 = (FVector4*)InterpolatedMotionVectorsPtr;
+
+						int32 Index = 0;
+						for (; Index + 3 < NumVerts; Index += 4)
+						{
+							VectorRegister MotionVector0xyz_MotionVector1x = VectorMultiplyAdd(VectorLoad(MotionVectorsAPtr4 + 0), WeightA, VectorMultiply(VectorLoad(MotionVectorsBPtr4 + 0), WeightB));
+							VectorRegister MotionVector1yz_MotionVector2xy = VectorMultiplyAdd(VectorLoad(MotionVectorsAPtr4 + 1), WeightA, VectorMultiply(VectorLoad(MotionVectorsBPtr4 + 1), WeightB));
+							VectorRegister MotionVector2z_MotionVector3xyz = VectorMultiplyAdd(VectorLoad(MotionVectorsAPtr4 + 2), WeightA, VectorMultiply(VectorLoad(MotionVectorsBPtr4 + 2), WeightB));
+							VectorStore(MotionVector0xyz_MotionVector1x, InterpolatedMotionVectorsPtr4 + 0);
+							VectorStore(MotionVector1yz_MotionVector2xy, InterpolatedMotionVectorsPtr4 + 1);
+							VectorStore(MotionVector2z_MotionVector3xyz, InterpolatedMotionVectorsPtr4 + 2);
+							MotionVectorsAPtr4 += 3;
+							MotionVectorsBPtr4 += 3;
+							InterpolatedMotionVectorsPtr4 += 3;
+						}
+
+						for (; Index < NumVerts; Index++)
 						{
-							InterpolatedMotionVectorsPtr[Index] = MotionVectorAPtr[Index] * OneMinusInterp + MotionVectorBPtr[Index] * InterpolationFactor;
+							InterpolatedMotionVectorsPtr[Index] = MotionVectorsAPtr[Index] * OneMinusInterp + MotionVectorsBPtr[Index] * InterpolationFactor;
 						}
 					}
+#if VALIDATE
+					for (int32 Index = 0; Index < NumVerts; ++Index)
+					{
+						FVector Result = MotionVectorsAPtr[Index] * OneMinusInterp + MotionVectorsBPtr[Index] * InterpolationFactor;
+						check(FMath::Abs(InterpolatedMotionVectorsPtr[Index].X - Result.X) < 0.01f);
+						check(FMath::Abs(InterpolatedMotionVectorsPtr[Index].Y - Result.Y) < 0.01f);
+						check(FMath::Abs(InterpolatedMotionVectorsPtr[Index].Z - Result.Z) < 0.01f);
+					}
+#endif
+				}
 
-					// Upload other non-motionblurred data
-					if (!TrackProxy->MeshData->VertexInfo.bConstantIndices)
-						TrackProxy->IndexBuffer.Update(TrackProxy->MeshData->Indices);
+#undef VALIDATE
 
-					if (TrackProxy->MeshData->VertexInfo.bHasTangentX)
-						TrackProxy->TangentXBuffer.Update(InterpolatedTangentX);
-					if (TrackProxy->MeshData->VertexInfo.bHasTangentZ)
-						TrackProxy->TangentZBuffer.Update(InterpolatedTangentZ);
+				// Upload other non-motionblurred data
+				if (!TrackProxy->MeshData->VertexInfo.bConstantIndices)
+					TrackProxy->IndexBuffer.Update(TrackProxy->MeshData->Indices);
 
-					if (TrackProxy->MeshData->VertexInfo.bHasUV0)
-						TrackProxy->TextureCoordinatesBuffer.Update(InterpolatedUVs);
+				if (TrackProxy->MeshData->VertexInfo.bHasTangentX)
+					TrackProxy->TangentXBuffer.Update(Scratch.InterpolatedTangentX);
+				if (TrackProxy->MeshData->VertexInfo.bHasTangentZ)
+					TrackProxy->TangentZBuffer.Update(Scratch.InterpolatedTangentZ);
 
-					if (TrackProxy->MeshData->VertexInfo.bHasColor0)
-						TrackProxy->ColorBuffer.Update(InterpolatedColors);
+				if (TrackProxy->MeshData->VertexInfo.bHasUV0)
+					TrackProxy->TextureCoordinatesBuffer.Update(Scratch.InterpolatedUVs);
 
-					bool bIsCompatibleWithCachedFrame = TrackProxy->IsTopologyCompatible(
-						TrackProxy->PositionBufferFrameIndices[TrackProxy->CurrentPositionBufferIndex % 2],
-						TrackProxy->FrameIndex);
+				if (TrackProxy->MeshData->VertexInfo.bHasColor0)
+					TrackProxy->ColorBuffer.Update(Scratch.InterpolatedColors);
 
-					if (!bHasMotionVectors)
-					{
-						// Initialize both buffers the first frame
-						if (TrackProxy->CurrentPositionBufferIndex == -1 || !bIsCompatibleWithCachedFrame)
-						{
-							TrackProxy->PositionBuffers[0].Update(InterpolatedPositions);
-							TrackProxy->PositionBuffers[1].Update(InterpolatedPositions);
-							TrackProxy->CurrentPositionBufferIndex = 0;
-							TrackProxy->PositionBufferFrameTimes[0] = Time;
-							TrackProxy->PositionBufferFrameTimes[1] = Time;
-							// We need to keep a frame index in order to ensure topology consistency. As we can interpolate 
-							// FrameIndex and NextFrameIndex are certainly topo-compatible so it doesn't really matter which 
-							// one we keep here. But wee keep NextFrameIndex as that is most useful to validate against
-							// the frame coming up
-							TrackProxy->PositionBufferFrameIndices[0] = TrackProxy->NextFrameIndex;
-							TrackProxy->PositionBufferFrameIndices[1] = TrackProxy->NextFrameIndex;
-						}
-						else
-						{
-							TrackProxy->CurrentPositionBufferIndex++;
-							TrackProxy->PositionBuffers[TrackProxy->CurrentPositionBufferIndex % 2].Update(InterpolatedPositions);
-							TrackProxy->PositionBufferFrameTimes[TrackProxy->CurrentPositionBufferIndex % 2] = Time;
-							TrackProxy->PositionBufferFrameIndices[TrackProxy->CurrentPositionBufferIndex % 2] = TrackProxy->NextFrameIndex;
-						}
-					}
-					else
+				bool bIsCompatibleWithCachedFrame = TrackProxy->IsTopologyCompatible(
+					TrackProxy->PositionBufferFrameIndices[TrackProxy->CurrentPositionBufferIndex % 2],
+					TrackProxy->FrameIndex);
+
+				if (!bHasMotionVectors)
+				{
+					// Initialize both buffers the first frame
+					if (TrackProxy->CurrentPositionBufferIndex == -1 || !bIsCompatibleWithCachedFrame)
 					{
+						TrackProxy->PositionBuffers[0].Update(Scratch.InterpolatedPositions);
+						TrackProxy->PositionBuffers[1].Update(Scratch.InterpolatedPositions);
 						TrackProxy->CurrentPositionBufferIndex = 0;
-						TrackProxy->PositionBuffers[0].Update(InterpolatedPositions);
-						TrackProxy->PositionBuffers[1].Update(InterpolatedMotionVectors);
-						TrackProxy->PositionBufferFrameIndices[0] = TrackProxy->FrameIndex;
-						TrackProxy->PositionBufferFrameIndices[1] = -1;
 						TrackProxy->PositionBufferFrameTimes[0] = Time;
 						TrackProxy->PositionBufferFrameTimes[1] = Time;
+						// We need to keep a frame index in order to ensure topology consistency. As we can interpolate 
+						// FrameIndex and NextFrameIndex are certainly topo-compatible so it doesn't really matter which 
+						// one we keep here. But wee keep NextFrameIndex as that is most useful to validate against
+						// the frame coming up
+						TrackProxy->PositionBufferFrameIndices[0] = TrackProxy->NextFrameIndex;
+						TrackProxy->PositionBufferFrameIndices[1] = TrackProxy->NextFrameIndex;
+					}
+					else
+					{
+						TrackProxy->CurrentPositionBufferIndex++;
+						TrackProxy->PositionBuffers[TrackProxy->CurrentPositionBufferIndex % 2].Update(Scratch.InterpolatedPositions);
+						TrackProxy->PositionBufferFrameTimes[TrackProxy->CurrentPositionBufferIndex % 2] = Time;
+						TrackProxy->PositionBufferFrameIndices[TrackProxy->CurrentPositionBufferIndex % 2] = TrackProxy->NextFrameIndex;
 					}
 				}
-
-
+				else
+				{
+					TrackProxy->CurrentPositionBufferIndex = 0;
+					TrackProxy->PositionBuffers[0].Update(Scratch.InterpolatedPositions);
+					TrackProxy->PositionBuffers[1].Update(Scratch.InterpolatedMotionVectors);
+					TrackProxy->PositionBufferFrameIndices[0] = TrackProxy->FrameIndex;
+					TrackProxy->PositionBufferFrameIndices[1] = -1;
+					TrackProxy->PositionBufferFrameTimes[0] = Time;
+					TrackProxy->PositionBufferFrameTimes[1] = Time;
+				}
 			}
 			else
 			{
 				// We just don't interpolate between frames if we got GPU to burn we could someday render twice and stipple fade between it :-D like with lods
 
 				// Only bother uploading if anything changed or when the we failed to decode anything make sure update the gpu buffers regardless
-				if (bFrameIndicesChanged || bDifferentInterpolationFactor || bDecodedAnything || bDecoderError)
+				if (bFrameIndicesChanged || bDifferentRoundedInterpolationFactor || bDecodedAnything || bDecoderError)
 				{
 					const bool bNextFrame = !!FMath::RoundToInt(InterpolationFactor);
 					const uint32 FrameIndexToUse = bNextFrame ? TrackProxy->NextFrameIndex : TrackProxy->FrameIndex;
@@ -912,7 +1025,7 @@ void FGeometryCacheSceneProxy::FrameUpdate() const
 						TrackProxy->ColorBuffer.Update(MeshDataToUse->Colors);
 
 					const bool bIsCompatibleWithCachedFrame = TrackProxy->IsTopologyCompatible(
-							TrackProxy->PositionBufferFrameIndices[TrackProxy->CurrentPositionBufferIndex % 2],
+						TrackProxy->PositionBufferFrameIndices[TrackProxy->CurrentPositionBufferIndex % 2],
 						FrameIndexToUse);
 
 					if (!bHasMotionVectors)
@@ -951,7 +1064,7 @@ void FGeometryCacheSceneProxy::FrameUpdate() const
 				}
 			}
 
-		#if 0
+#if 0
 			bool bOffloadUpdate = CVarOffloadUpdate.GetValueOnRenderThread() != 0;
 			if (TrackProxy->SampleIndex != TrackProxy->UploadedSampleIndex)
 			{
@@ -971,7 +1084,7 @@ void FGeometryCacheSceneProxy::FrameUpdate() const
 					}, GET_STATID(STAT_BufferUpdateTask), NULL, ENamedThreads::AnyThread);
 
 					// Queue a command on the RHI thread that waits for the interpolation job and then uploads them to the GPU
-					FRHICommandListImmediate &RHICommandList = GetImmediateCommandList_ForRenderCommand();
+					FRHICommandListImmediate& RHICommandList = GetImmediateCommandList_ForRenderCommand();
 					new (RHICommandList.AllocCommand<FRHICommandUpdateGeometryCacheBuffer>())FRHICommandUpdateGeometryCacheBuffer(
 						CompletionFence,
 						TrackProxy->VertexBuffer.VertexBufferRHI,
@@ -1006,11 +1119,11 @@ void FGeometryCacheSceneProxy::FrameUpdate() const
 					else
 					{
 						TrackProxy->CurrentPositionBufferIndex++;
-						TrackProxy->PositonBuffers[TrackProxy->CurrentPositionBufferIndex%2].Update(TrackProxy->MeshData->Vertices);
+						TrackProxy->PositonBuffers[TrackProxy->CurrentPositionBufferIndex % 2].Update(TrackProxy->MeshData->Vertices);
 					}
 				}
 			}
-		#endif
+#endif
 
 		}
 	}
@@ -1025,6 +1138,7 @@ void FGeometryCacheSceneProxy::UpdateSectionWorldMatrix(const int32 SectionIndex
 void FGeometryCacheSceneProxy::ClearSections()
 {
 	Tracks.Empty();
+	Scratch.Empty();
 }
 
 bool FGeomCacheTrackProxy::UpdateMeshData(float Time, bool bLooping, int32& InOutMeshSampleIndex, FGeometryCacheMeshData& OutMeshData)
@@ -1063,7 +1177,7 @@ const FVisibilitySample& FGeomCacheTrackProxy::GetVisibilitySample(float Time, c
 	return FVisibilitySample::InvisibleSample;
 }
 
-void FGeomCacheTrackProxy::FindSampleIndexesFromTime(float Time, bool bLooping, bool bIsPlayingBackwards, int32 &OutFrameIndex, int32 &OutNextFrameIndex, float& InInterpolationFactor)
+void FGeomCacheTrackProxy::FindSampleIndexesFromTime(float Time, bool bLooping, bool bIsPlayingBackwards, int32& OutFrameIndex, int32& OutNextFrameIndex, float& InInterpolationFactor)
 {
 	if (UGeometryCacheTrackStreamable* StreamableTrack = Cast<UGeometryCacheTrackStreamable>(Track))
 	{
@@ -1105,7 +1219,7 @@ void FGeomCacheVertexFactory::Init(const FVertexBuffer* PositionBuffer, const FV
 
 		ENQUEUE_RENDER_COMMAND(InitGeomCacheVertexFactory)(
 			[this, PositionBuffer, MotionBlurDataBuffer, TangentXBuffer, TangentZBuffer, TextureCoordinateBuffer, ColorBuffer](FRHICommandListImmediate& RHICmdList)
-			{
+		{
 			Init_RenderThread(PositionBuffer, MotionBlurDataBuffer, TangentXBuffer, TangentZBuffer, TextureCoordinateBuffer, ColorBuffer);
 		});
 	}
@@ -1119,7 +1233,7 @@ void FGeomCacheIndexBuffer::InitRHI()
 	RHIUnlockIndexBuffer(IndexBufferRHI);
 }
 
-void FGeomCacheIndexBuffer::Update(const TArray<uint32> &Indices)
+void FGeomCacheIndexBuffer::Update(const TArray<uint32>& Indices)
 {
 	SCOPE_CYCLE_COUNTER(STAT_IndexBufferUpdate);
 
@@ -1165,7 +1279,7 @@ void FGeomCacheVertexBuffer::InitRHI()
 	RHIUnlockVertexBuffer(VertexBufferRHI);
 }
 
-void FGeomCacheVertexBuffer::UpdateRaw(const void *Data, int32 NumItems, int32 ItemSizeBytes, int32 ItemStrideBytes)
+void FGeomCacheVertexBuffer::UpdateRaw(const void* Data, int32 NumItems, int32 ItemSizeBytes, int32 ItemStrideBytes)
 {
 	SCOPE_CYCLE_COUNTER(STAT_VertexBufferUpdate);
 	int32 NewSizeInBytes = ItemSizeBytes * NumItems;
@@ -1190,9 +1304,9 @@ void FGeomCacheVertexBuffer::UpdateRaw(const void *Data, int32 NumItems, int32 I
 	}
 	else
 	{
-		int8 *InBytes = (int8 *)Data;
-		int8 *OutBytes = (int8 *)VertexBufferData;
-		for (int32 ItemId=0; ItemId < NumItems; ItemId++)
+		int8* InBytes = (int8*)Data;
+		int8* OutBytes = (int8*)VertexBufferData;
+		for (int32 ItemId = 0; ItemId < NumItems; ItemId++)
 		{
 			FMemory::Memcpy(OutBytes, InBytes, ItemSizeBytes);
 			InBytes += ItemStrideBytes;
diff --git a/Engine/Plugins/Experimental/GeometryCache/Source/GeometryCache/Public/GeometryCacheSceneProxy.h b/Engine/Plugins/Experimental/GeometryCache/Source/GeometryCache/Public/GeometryCacheSceneProxy.h
--- a/Engine/Plugins/Experimental/GeometryCache/Source/GeometryCache/Public/GeometryCacheSceneProxy.h
+++ b/Engine/Plugins/Experimental/GeometryCache/Source/GeometryCache/Public/GeometryCacheSceneProxy.h
@@ -27,7 +27,7 @@ public:
 	void Init(int32 InSizeInBytes)
 	{
 		check(this->IsInitialized() == false);
-		SizeInBytes = InSizeInBytes;		
+		SizeInBytes = InSizeInBytes;
 	}
 
 	/* Create on rhi thread. Uninitialized with NumVertices space.*/
@@ -36,20 +36,20 @@ public:
 	/**
 	 * Sugar function to update from a typed array.
 	 */
-	template<class DataType> void Update(const TArray<DataType> &Vertices)
+	template<class DataType> void Update(const TArray<DataType>& Vertices)
 	{
 		int32 InSize = Vertices.Num() * sizeof(DataType);
 		UpdateRaw(Vertices.GetData(), InSize, 1, 1);
 	}
 
-	void UpdatePositionsOnly(const TArray<FDynamicMeshVertex> &Vertices)
+	void UpdatePositionsOnly(const TArray<FDynamicMeshVertex>& Vertices)
 	{
 		const uint32 PositionOffset = STRUCT_OFFSET(FDynamicMeshVertex, Position);
 		const uint32 PositionSize = sizeof(((FDynamicMeshVertex*)nullptr)->Position);
 		UpdateRaw(Vertices.GetData() + PositionOffset, Vertices.Num(), PositionSize, sizeof(FDynamicMeshVertex));
 	}
 
-	void UpdateExceptPositions(const TArray<FDynamicMeshVertex> &Vertices)
+	void UpdateExceptPositions(const TArray<FDynamicMeshVertex>& Vertices)
 	{
 		const uint32 PositionSize = sizeof(((FDynamicMeshVertex*)nullptr)->Position);
 		const uint32 PositionOffset = STRUCT_OFFSET(FDynamicMeshVertex, Position);
@@ -57,13 +57,13 @@ public:
 		static_assert(PositionOffset == 0, "Expecting position to be the first struct member");
 		static_assert(PositionSize == STRUCT_OFFSET(FDynamicMeshVertex, TextureCoordinate), "Expecting the texture coordinate to immediately follow the Position");
 
-		UpdateRaw( (int8 *)Vertices.GetData() + PositionSize, Vertices.Num(), sizeof(FDynamicMeshVertex) - PositionSize, sizeof(FDynamicMeshVertex));
+		UpdateRaw((int8*)Vertices.GetData() + PositionSize, Vertices.Num(), sizeof(FDynamicMeshVertex) - PositionSize, sizeof(FDynamicMeshVertex));
 	}
 
 	/**
 	 * Update the raw contents of the buffer, possibly reallocate if needed.
 	 */
-	void UpdateRaw(const void *Data, int32 NumItems, int32 ItemSizeBytes, int32 ItemStrideBytes);
+	void UpdateRaw(const void* Data, int32 NumItems, int32 ItemSizeBytes, int32 ItemStrideBytes);
 
 	/**
 	 * Resize the buffer but don't initialize it with any data.
@@ -101,7 +101,7 @@ public:
 	/**
 		Update the data and possibly reallocate if needed.
 	*/
-	void Update(const TArray<uint32> &Indices);
+	void Update(const TArray<uint32>& Indices);
 
 	void UpdateSizeOnly(int32 NewNumIndices);
 
@@ -192,11 +192,11 @@ public:
 	 * @param OutNextFrameIndex - The frame index that follows OutFrameIndex
 	 * @param InterpolationFactor - The interpolation value between the two frame times
 	 */
-	virtual void FindSampleIndexesFromTime(float Time, bool bLooping, bool bIsPlayingBackwards, int32 &OutFrameIndex, int32 &OutNextFrameIndex, float &InInterpolationFactor);
+	virtual void FindSampleIndexesFromTime(float Time, bool bLooping, bool bIsPlayingBackwards, int32& OutFrameIndex, int32& OutNextFrameIndex, float& InInterpolationFactor);
 
 	/** MeshData storing information used for rendering this Track */
-	FGeometryCacheMeshData *MeshData;
-	FGeometryCacheMeshData *NextFrameMeshData;
+	FGeometryCacheMeshData* MeshData;
+	FGeometryCacheMeshData* NextFrameMeshData;
 
 	/** Frame numbers corresponding to MeshData, NextFrameMeshData */
 	int32 FrameIndex;
@@ -205,7 +205,7 @@ public:
 
 	/** Material applied to this Track */
 	TArray<UMaterialInterface*> Materials;
-	
+
 	/** Vertex buffers for this Track. There are two position buffers which we double buffer between, current frame and last frame*/
 	FGeomCacheVertexBuffer PositionBuffers[2];
 	uint32 PositionBufferFrameIndices[2]; // Frame indexes of the positions in the position buffer 
@@ -222,7 +222,7 @@ public:
 
 	/** Vertex factory for this Track */
 	FGeomCacheVertexFactory VertexFactory;
-	
+
 	/** World Matrix for this Track */
 	FMatrix WorldMatrix;
 
@@ -262,7 +262,7 @@ public:
 	/** Update world matrix for specific section */
 	void UpdateSectionWorldMatrix(const int32 SectionIndex, const FMatrix& WorldMatrix);
 	/** Update vertex buffer for specific section */
-	void UpdateSectionVertexBuffer(const int32 SectionIndex, FGeometryCacheMeshData* MeshData );
+	void UpdateSectionVertexBuffer(const int32 SectionIndex, FGeometryCacheMeshData* MeshData);
 	/** Update index buffer for specific section */
 	void UpdateSectionIndexBuffer(const int32 SectionIndex, const TArray<uint32>& Indices);
 
@@ -273,18 +273,8 @@ public:
 	virtual void GetDynamicRayTracingInstances(FRayTracingMaterialGatheringContext& Context, TArray<FRayTracingInstance>& OutRayTracingInstances) override final;
 	virtual bool IsRayTracingRelevant() const override { return true; }
 #endif
-	
-private:	
-	FMaterialRelevance MaterialRelevance;
-	float Time;
-	bool bLooping;
-	bool bIsPlayingBackwards;
-	float PlaybackSpeed;
-	uint32 UpdatedFrameNum;
-
-	/** Array of Track Proxies */
-	TArray<FGeomCacheTrackProxy*> Tracks;
 
+private:
 	void FrameUpdate() const;
 
 	void CreateMeshBatch(
@@ -294,6 +284,71 @@ private:
 		FDynamicPrimitiveUniformBuffer& DynamicPrimitiveUniformBuffer,
 		FMeshBatch& Mesh) const;
 
+private:
+	/** Array of Track Proxies */
+	TArray<FGeomCacheTrackProxy*> Tracks;
+
+	/** Scratch memory for frame update - do not use directly. */
+	struct FScratchMemory
+	{
+		TArray<FVector> InterpolatedPositions;
+		TArray<FPackedNormal> InterpolatedTangentX;
+		TArray<FPackedNormal> InterpolatedTangentZ;
+		TArray<FVector2D> InterpolatedUVs;
+		TArray<FColor> InterpolatedColors;
+		TArray<FVector> InterpolatedMotionVectors;
+
+		void Prepare(SIZE_T NumVertices, bool bHasMotionVectors)
+		{
+			// Clear entries but keep allocations.
+			InterpolatedPositions.Reset();
+			InterpolatedTangentX.Reset();
+			InterpolatedTangentZ.Reset();
+			InterpolatedUVs.Reset();
+			InterpolatedColors.Reset();
+			InterpolatedMotionVectors.Reset();
+
+			// Make sure our capacity fits the requested vertex count
+			InterpolatedPositions.Reserve(NumVertices);
+			InterpolatedTangentX.Reserve(NumVertices);
+			InterpolatedTangentZ.Reserve(NumVertices);
+			InterpolatedUVs.Reserve(NumVertices);
+			InterpolatedColors.Reserve(NumVertices);
+
+			InterpolatedPositions.AddUninitialized(NumVertices);
+			InterpolatedTangentX.AddUninitialized(NumVertices);
+			InterpolatedTangentZ.AddUninitialized(NumVertices);
+			InterpolatedUVs.AddUninitialized(NumVertices);
+			InterpolatedColors.AddUninitialized(NumVertices);
+
+			if (bHasMotionVectors)
+			{
+				InterpolatedMotionVectors.Reserve(NumVertices);
+				InterpolatedMotionVectors.AddUninitialized(NumVertices);
+			}
+		}
+
+		void Empty()
+		{
+			// Clear entries but and release memory.
+			InterpolatedPositions.Empty();
+			InterpolatedTangentX.Empty();
+			InterpolatedTangentZ.Empty();
+			InterpolatedUVs.Empty();
+			InterpolatedColors.Empty();
+			InterpolatedMotionVectors.Empty();
+		}
+	}
+	mutable Scratch;
+
+	uint32 UpdatedFrameNum;
+	float Time;
+	float PlaybackSpeed;
+
+	FMaterialRelevance MaterialRelevance;
+	uint32 bLooping : 1;
+	uint32 bIsPlayingBackwards : 1;
+
 	/** Function used to create a new track proxy at construction */
 	TFunction<FGeomCacheTrackProxy*()> CreateTrackProxy;
 };
