GeometryCacheSceneProxy needs to correctly mark lighting channel view relevance.

diff --git a/Engine/Plugins/Experimental/GeometryCache/Source/GeometryCache/Private/GeometryCacheSceneProxy.cpp b/Engine/Plugins/Experimental/GeometryCache/Source/GeometryCache/Private/GeometryCacheSceneProxy.cpp
--- a/Engine/Plugins/Experimental/GeometryCache/Source/GeometryCache/Private/GeometryCacheSceneProxy.cpp
+++ b/Engine/Plugins/Experimental/GeometryCache/Source/GeometryCache/Private/GeometryCacheSceneProxy.cpp
@@ -530,6 +530,7 @@ FPrimitiveViewRelevance FGeometryCacheSceneProxy::GetViewRelevance(const FSceneV
 	Result.bShadowRelevance = IsShadowCast(View);
 	Result.bDynamicRelevance = true;
 	Result.bRenderCustomDepth = ShouldRenderCustomDepth();
+	Result.bUsesLightingChannels = GetLightingChannelMask() != GetDefaultLightingChannelMask();
 	MaterialRelevance.SetPrimitiveViewRelevance(Result);
 	Result.bVelocityRelevance = IsMovable() && Result.bOpaque && Result.bRenderInMainPass;
 	return Result;
