Adding EditorFramework to build.cs files

#rnx

diff --git a/Engine/Plugins/2D/Paper2D/Source/Paper2D/Paper2D.Build.cs b/Engine/Plugins/2D/Paper2D/Source/Paper2D/Paper2D.Build.cs
--- a/Engine/Plugins/2D/Paper2D/Source/Paper2D/Paper2D.Build.cs
+++ b/Engine/Plugins/2D/Paper2D/Source/Paper2D/Paper2D.Build.cs
@@ -31,6 +31,7 @@ public class Paper2D : ModuleRules
 		{
 			//@TODO: Needed for the triangulation code used for sprites (but only in editor mode)
 			//@TOOD: Try to move the code dependent on the triangulation code to the editor-only module
+			PrivateDependencyModuleNames.Add("EditorFramework");
 			PrivateDependencyModuleNames.Add("UnrealEd");
 		}
 	}
diff --git a/Engine/Plugins/2D/Paper2D/Source/PaperSpriteSheetImporter/PaperSpriteSheetImporter.Build.cs b/Engine/Plugins/2D/Paper2D/Source/PaperSpriteSheetImporter/PaperSpriteSheetImporter.Build.cs
--- a/Engine/Plugins/2D/Paper2D/Source/PaperSpriteSheetImporter/PaperSpriteSheetImporter.Build.cs
+++ b/Engine/Plugins/2D/Paper2D/Source/PaperSpriteSheetImporter/PaperSpriteSheetImporter.Build.cs
@@ -15,6 +15,7 @@ public class PaperSpriteSheetImporter : ModuleRules
 				"SlateCore",
 				"Engine",
 				"Paper2D",
+				"EditorFramework",
 				"UnrealEd",
 				"Paper2DEditor",
 				"AssetTools",
diff --git a/Engine/Plugins/2D/Paper2D/Source/PaperTiledImporter/PaperTiledImporter.Build.cs b/Engine/Plugins/2D/Paper2D/Source/PaperTiledImporter/PaperTiledImporter.Build.cs
--- a/Engine/Plugins/2D/Paper2D/Source/PaperTiledImporter/PaperTiledImporter.Build.cs
+++ b/Engine/Plugins/2D/Paper2D/Source/PaperTiledImporter/PaperTiledImporter.Build.cs
@@ -13,6 +13,7 @@ public class PaperTiledImporter : ModuleRules
 				"Engine",
 				"Json",
 				"Paper2D",
+				"EditorFramework",
 				"UnrealEd",
 				"Paper2DEditor"
 			});
diff --git a/Engine/Plugins/Developer/RenderDocPlugin/Source/RenderDocPlugin/RenderDocPlugin.Build.cs b/Engine/Plugins/Developer/RenderDocPlugin/Source/RenderDocPlugin/RenderDocPlugin.Build.cs
--- a/Engine/Plugins/Developer/RenderDocPlugin/Source/RenderDocPlugin/RenderDocPlugin.Build.cs
+++ b/Engine/Plugins/Developer/RenderDocPlugin/Source/RenderDocPlugin/RenderDocPlugin.Build.cs
@@ -32,6 +32,7 @@ namespace UnrealBuildTool.Rules
 					"Slate"
 					, "SlateCore"
 					, "EditorStyle"
+					, "EditorFramework"
 					, "UnrealEd"
 					, "MainFrame"
 					, "GameProjectGeneration"
diff --git a/Engine/Plugins/Experimental/GeometryCache/Source/GeometryCacheEd/GeometryCacheEd.Build.cs b/Engine/Plugins/Experimental/GeometryCache/Source/GeometryCacheEd/GeometryCacheEd.Build.cs
--- a/Engine/Plugins/Experimental/GeometryCache/Source/GeometryCacheEd/GeometryCacheEd.Build.cs
+++ b/Engine/Plugins/Experimental/GeometryCache/Source/GeometryCacheEd/GeometryCacheEd.Build.cs
@@ -15,7 +15,8 @@ public class GeometryCacheEd : ModuleRules
 				"SlateCore",
                 "InputCore",
                 "RenderCore",
-                "RHI",		
+                "RHI",
+				"EditorFramework",
                 "UnrealEd",
 				"AssetTools",
                 "GeometryCache",
diff --git a/Engine/Source/Developer/MaterialBaking/MaterialBaking.Build.cs b/Engine/Source/Developer/MaterialBaking/MaterialBaking.Build.cs
--- a/Engine/Source/Developer/MaterialBaking/MaterialBaking.Build.cs
+++ b/Engine/Source/Developer/MaterialBaking/MaterialBaking.Build.cs
@@ -12,7 +12,8 @@ public class MaterialBaking : ModuleRules
 				"CoreUObject",
 				"Engine",
 				"RenderCore",
-                "RHI",                
+                "RHI",        
+				"EditorFramework",        
                 "UnrealEd",
                 "MainFrame",
                 "SlateCore",
