Fixing PVS 7.7 Issues:

warning V610: Undefined behavior. Check the shift operator '<<'. The left operand is negative ('(Sign ? - 2 : 1)' = [-2..1]).

diff --git a/Engine/Plugins/Experimental/GeometryCache/Source/GeometryCache/Private/CodecV1.cpp b/Engine/Plugins/Experimental/GeometryCache/Source/GeometryCache/Private/CodecV1.cpp
--- a/Engine/Plugins/Experimental/GeometryCache/Source/GeometryCache/Private/CodecV1.cpp
+++ b/Engine/Plugins/Experimental/GeometryCache/Source/GeometryCache/Private/CodecV1.cpp
@@ -1041,7 +1041,7 @@ void FCodecV1Decoder::InitLUT()
 	{
 		for (int32 Sign = 0; Sign <= 1; Sign++)
 		{
-			CachedHighBitsLUT[2 + Sign + NumRawBits * 2] = (Sign ? -2 : 1) << NumRawBits;
+			CachedHighBitsLUT[2 + Sign + NumRawBits * 2] = int32(uint32(Sign ? -2 : 1) << NumRawBits);
 		}
 	}
 }
