Replace CreateAndLock calls with Create and Lock

diff --git a/Engine/Plugins/Enterprise/LidarPointCloud/Source/LidarPointCloudRuntime/Private/Rendering/LidarPointCloudRenderBuffers.cpp b/Engine/Plugins/Enterprise/LidarPointCloud/Source/LidarPointCloudRuntime/Private/Rendering/LidarPointCloudRenderBuffers.cpp
--- a/Engine/Plugins/Enterprise/LidarPointCloud/Source/LidarPointCloudRuntime/Private/Rendering/LidarPointCloudRenderBuffers.cpp
+++ b/Engine/Plugins/Enterprise/LidarPointCloud/Source/LidarPointCloudRuntime/Private/Rendering/LidarPointCloudRenderBuffers.cpp
@@ -39,14 +39,13 @@ void FLidarPointCloudIndexBuffer::Resize(const uint32 & RequestedCapacity)
 
 void FLidarPointCloudIndexBuffer::InitRHI()
 {
-	FRHIResourceCreateInfo CreateInfo;
-	void* Buffer = nullptr;
+	FRHIResourceCreateInfo CreateInfo(TEXT("FLidarPointCloudIndexBuffer"));
 	uint32 Size = Capacity * 7 * sizeof(uint32);
 	PointOffset = Capacity * 6;
 
-	IndexBufferRHI = RHICreateAndLockIndexBuffer(sizeof(uint32), Size, BUF_Dynamic, CreateInfo, Buffer);
+	IndexBufferRHI = RHICreateBuffer(Size, BUF_Dynamic | BUF_IndexBuffer, sizeof(uint32), ERHIAccess::VertexOrIndexBuffer, CreateInfo);
 
-	uint32* Data = (uint32*)Buffer;
+	uint32* Data = (uint32*)RHILockBuffer(IndexBufferRHI, 0, Size, RLM_WriteOnly);
 	for (uint32 i = 0; i < Capacity; i++)
 	{
 		// Full quads
@@ -65,9 +64,7 @@ void FLidarPointCloudIndexBuffer::InitRHI()
 		// Points
 		Data[PointOffset + i] = i;
 	}
-
 	RHIUnlockBuffer(IndexBufferRHI);
-	Buffer = nullptr;
 }
 
 //////////////////////////////////////////////////////////// Structured Buffer
diff --git a/Engine/Plugins/Enterprise/LidarPointCloud/Source/LidarPointCloudRuntime/Private/Rendering/LidarPointCloudRenderBuffers.h b/Engine/Plugins/Enterprise/LidarPointCloud/Source/LidarPointCloudRuntime/Private/Rendering/LidarPointCloudRenderBuffers.h
--- a/Engine/Plugins/Enterprise/LidarPointCloud/Source/LidarPointCloudRuntime/Private/Rendering/LidarPointCloudRenderBuffers.h
+++ b/Engine/Plugins/Enterprise/LidarPointCloud/Source/LidarPointCloudRuntime/Private/Rendering/LidarPointCloudRenderBuffers.h
@@ -162,12 +162,11 @@ private:
 	public:
 		virtual void InitRHI() override
 		{
-			FRHIResourceCreateInfo CreateInfo;
-			void* Buffer = nullptr;
-			VertexBufferRHI = RHICreateAndLockVertexBuffer(sizeof(FVector), BUF_Static | BUF_ZeroStride, CreateInfo, Buffer);
+			FRHIResourceCreateInfo CreateInfo(TEXT("FPointCloudVertexBuffer"));
+			VertexBufferRHI = RHICreateBuffer(sizeof(FVector), BUF_Static | BUF_VertexBuffer | BUF_ZeroStride, 0, ERHIAccess::VertexOrIndexBuffer, CreateInfo);
+			void* Buffer = RHILockBuffer(VertexBufferRHI, 0, sizeof(FVector), RLM_WriteOnly);
 			FMemory::Memzero(Buffer, sizeof(FVector));
 			RHIUnlockBuffer(VertexBufferRHI);
-			Buffer = nullptr;
 		}
 
 		virtual FString GetFriendlyName() const override { return TEXT("FPointCloudVertexBuffer"); }
diff --git a/Engine/Plugins/Enterprise/LidarPointCloud/Source/LidarPointCloudRuntime/Private/Rendering/LidarPointCloudRendering.cpp b/Engine/Plugins/Enterprise/LidarPointCloud/Source/LidarPointCloudRuntime/Private/Rendering/LidarPointCloudRendering.cpp
--- a/Engine/Plugins/Enterprise/LidarPointCloud/Source/LidarPointCloudRuntime/Private/Rendering/LidarPointCloudRendering.cpp
+++ b/Engine/Plugins/Enterprise/LidarPointCloud/Source/LidarPointCloudRuntime/Private/Rendering/LidarPointCloudRendering.cpp
@@ -103,12 +103,13 @@ private:
 
 		virtual void InitRHI() override
 		{
-			FRHIResourceCreateInfo CreateInfo;
-			void* Buffer = nullptr;
-			VertexBufferRHI = RHICreateAndLockVertexBuffer(DataLength * sizeof(FVector), BUF_Static | BUF_ShaderResource, CreateInfo, Buffer);
-			FMemory::Memcpy(Buffer, Data, DataLength * sizeof(FVector));
+			const uint32 Size = DataLength * sizeof(FVector);
+
+			FRHIResourceCreateInfo CreateInfo(TEXT("FLidarPointCloudCollisionVertexBuffer"));
+			VertexBufferRHI = RHICreateBuffer(Size, BUF_Static | BUF_VertexBuffer | BUF_ShaderResource, 0, ERHIAccess::VertexOrIndexBuffer | ERHIAccess::SRVMask, CreateInfo);
+			void* Buffer = RHILockBuffer(VertexBufferRHI, 0, Size, RLM_WriteOnly);
+			FMemory::Memcpy(Buffer, Data, Size);
 			RHIUnlockBuffer(VertexBufferRHI);
-			Buffer = nullptr;
 		}
 	} VertexBuffer;
 	class FLidarPointCloudCollisionIndexBuffer : public FIndexBuffer
@@ -128,12 +129,13 @@ private:
 
 		virtual void InitRHI() override
 		{
-			FRHIResourceCreateInfo CreateInfo;
-			void* Buffer = nullptr;
-			IndexBufferRHI = RHICreateAndLockIndexBuffer(sizeof(uint32), DataLength * sizeof(uint32), BUF_Static, CreateInfo, Buffer);
-			FMemory::Memcpy(Buffer, Data, DataLength * sizeof(uint32));
+			const uint32 Size = DataLength * sizeof(uint32);
+
+			FRHIResourceCreateInfo CreateInfo(TEXT("FLidarPointCloudCollisionIndexBuffer"));
+			IndexBufferRHI = RHICreateBuffer(Size, BUF_Static | BUF_IndexBuffer, sizeof(uint32), ERHIAccess::VertexOrIndexBuffer, CreateInfo);
+			void* Buffer = RHILockBuffer(IndexBufferRHI, 0, Size, RLM_WriteOnly);
+			FMemory::Memcpy(Buffer, Data, Size);
 			RHIUnlockBuffer(IndexBufferRHI);
-			Buffer = nullptr;
 		}
 	} IndexBuffer;
 
diff --git a/Engine/Plugins/Experimental/AR/PointCloud/Source/PointCloud/Private/PointCloudBuffers.h b/Engine/Plugins/Experimental/AR/PointCloud/Source/PointCloud/Private/PointCloudBuffers.h
--- a/Engine/Plugins/Experimental/AR/PointCloud/Source/PointCloud/Private/PointCloudBuffers.h
+++ b/Engine/Plugins/Experimental/AR/PointCloud/Source/PointCloud/Private/PointCloudBuffers.h
@@ -130,10 +130,9 @@ public:
 		NumPrimitives = NumPoints;
 		const uint32 Size = sizeof(INDEX_TYPE) * 4 * NumPoints;
 		const uint32 Stride = sizeof(INDEX_TYPE);
-		FRHIResourceCreateInfo CreateInfo;
-		void* Buffer = nullptr;
-		IndexBufferRHI = RHICreateAndLockIndexBuffer( Stride, Size, BUF_Static, CreateInfo, Buffer );
-		INDEX_TYPE* Indices = (INDEX_TYPE*)Buffer;
+		FRHIResourceCreateInfo CreateInfo(TEXT("PointCloudQuadList"));
+		IndexBufferRHI = RHICreateBuffer( Size, BUF_Static | BUF_IndexBuffer, Stride, ERHIAccess::VertexOrIndexBuffer, CreateInfo );
+		INDEX_TYPE* Indices = (INDEX_TYPE*)RHILockBuffer( IndexBufferRHI, 0, Size, RLM_WriteOnly );
 		for (uint32 SpriteIndex = 0; SpriteIndex < NumPoints; ++SpriteIndex)
 		{
 			Indices[SpriteIndex * 4 + 0] = SpriteIndex * 4 + 0;
@@ -151,10 +150,9 @@ public:
 		NumPrimitives = 2 * NumPoints;
 		const uint32 Size = sizeof(INDEX_TYPE) * 6 * NumPoints;
 		const uint32 Stride = sizeof(INDEX_TYPE);
-		FRHIResourceCreateInfo CreateInfo;
-		void* Buffer = nullptr;
-		IndexBufferRHI = RHICreateAndLockIndexBuffer( Stride, Size, BUF_Static, CreateInfo, Buffer );
-		INDEX_TYPE* Indices = (INDEX_TYPE*)Buffer;
+		FRHIResourceCreateInfo CreateInfo(TEXT("PointCloudTriList"));
+		IndexBufferRHI = RHICreateBuffer( Size, BUF_Static | BUF_IndexBuffer, Stride, ERHIAccess::VertexOrIndexBuffer, CreateInfo );
+		INDEX_TYPE* Indices = (INDEX_TYPE*)RHILockBuffer( IndexBufferRHI, 0, Size, RLM_WriteOnly );
 		for (uint32 SpriteIndex = 0; SpriteIndex < NumPoints; ++SpriteIndex)
 		{
 			Indices[SpriteIndex * 6 + 0] = SpriteIndex * 4 + 0;
diff --git a/Engine/Plugins/Experimental/AR/PointCloud/Source/PointCloud/Private/PointCloudVertexFactory.cpp b/Engine/Plugins/Experimental/AR/PointCloud/Source/PointCloud/Private/PointCloudVertexFactory.cpp
--- a/Engine/Plugins/Experimental/AR/PointCloud/Source/PointCloud/Private/PointCloudVertexFactory.cpp
+++ b/Engine/Plugins/Experimental/AR/PointCloud/Source/PointCloud/Private/PointCloudVertexFactory.cpp
@@ -84,10 +84,9 @@ public:
 
 	virtual void InitRHI() override
 	{
-		FRHIResourceCreateInfo CreateInfo;
-		void* BufferData = nullptr;
-		VertexBufferRHI = RHICreateAndLockVertexBuffer(sizeof(FVector4) * 4, BUF_Static, CreateInfo, BufferData);
-		FVector4* DummyContents = (FVector4*)BufferData;
+		FRHIResourceCreateInfo CreateInfo(TEXT("FDummyVertexBuffer"));
+		VertexBufferRHI = RHICreateBuffer(sizeof(FVector4) * 4, BUF_Static | BUF_VertexBuffer, 0, ERHIAccess::VertexOrIndexBuffer, CreateInfo);
+		FVector4* DummyContents = (FVector4*)RHILockBuffer(VertexBufferRHI, 0, sizeof(FVector) * 4, RLM_WriteOnly);
 //@todo - joeg do I need a quad's worth?
 		DummyContents[0] = FVector4(0.0f, 0.0f, 0.0f, 0.0f);
 		DummyContents[1] = FVector4(1.0f, 0.0f, 0.0f, 0.0f);
diff --git a/Engine/Plugins/Experimental/GeometryCache/Source/GeometryCache/Private/GeometryCacheSceneProxy.cpp b/Engine/Plugins/Experimental/GeometryCache/Source/GeometryCache/Private/GeometryCacheSceneProxy.cpp
--- a/Engine/Plugins/Experimental/GeometryCache/Source/GeometryCache/Private/GeometryCacheSceneProxy.cpp
+++ b/Engine/Plugins/Experimental/GeometryCache/Source/GeometryCache/Private/GeometryCacheSceneProxy.cpp
@@ -1275,10 +1275,8 @@ void FGeomCacheVertexFactory::Init(const FVertexBuffer* PositionBuffer, const FV
 
 void FGeomCacheIndexBuffer::InitRHI()
 {
-	FRHIResourceCreateInfo CreateInfo;
-	void* Buffer = nullptr;
-	IndexBufferRHI = RHICreateAndLockIndexBuffer(sizeof(uint32), NumIndices * sizeof(uint32), BUF_Dynamic | BUF_ShaderResource, CreateInfo, Buffer);
-	RHIUnlockBuffer(IndexBufferRHI);
+	FRHIResourceCreateInfo CreateInfo(TEXT("FGeomCacheIndexBuffer"));
+	IndexBufferRHI = RHICreateBuffer(NumIndices * sizeof(uint32), BUF_Dynamic | BUF_IndexBuffer | BUF_ShaderResource, sizeof(uint32), ERHIAccess::VertexOrIndexBuffer | ERHIAccess::SRVMask, CreateInfo);
 }
 
 void FGeomCacheIndexBuffer::Update(const TArray<uint32>& Indices)
@@ -1293,16 +1291,14 @@ void FGeomCacheIndexBuffer::Update(const TArray<uint32>& Indices)
 	if (Indices.Num() > NumIndices)
 	{
 		NumIndices = Indices.Num();
-		FRHIResourceCreateInfo CreateInfo;
-		IndexBufferRHI = RHICreateAndLockIndexBuffer(sizeof(uint32), NumIndices * sizeof(uint32), BUF_Dynamic | BUF_ShaderResource, CreateInfo, Buffer);
+		FRHIResourceCreateInfo CreateInfo(TEXT("FGeomCacheIndexBuffer"));
+		IndexBufferRHI = RHICreateBuffer(NumIndices * sizeof(uint32), BUF_Dynamic | BUF_IndexBuffer | BUF_ShaderResource, sizeof(uint32), ERHIAccess::VertexOrIndexBuffer | ERHIAccess::SRVMask, CreateInfo);
 	}
-	else
+
+	if (Indices.Num() > 0)
 	{
-		if (Indices.Num() > 0)
-		{
-			// Copy the index data into the index buffer.
-			Buffer = RHILockBuffer(IndexBufferRHI, 0, Indices.Num() * sizeof(uint32), RLM_WriteOnly);
-		}
+		// Copy the index data into the index buffer.
+		Buffer = RHILockBuffer(IndexBufferRHI, 0, Indices.Num() * sizeof(uint32), RLM_WriteOnly);
 	}
 
 	if (Buffer)
@@ -1319,18 +1315,16 @@ void FGeomCacheIndexBuffer::UpdateSizeOnly(int32 NewNumIndices)
 	// We only ever grow in size. Ok for now?
 	if (NewNumIndices > NumIndices)
 	{
-		FRHIResourceCreateInfo CreateInfo;
-		IndexBufferRHI = RHICreateIndexBuffer(sizeof(uint32), NewNumIndices * sizeof(uint32), BUF_Dynamic | BUF_ShaderResource, CreateInfo);
+		FRHIResourceCreateInfo CreateInfo(TEXT("FGeomCacheIndexBuffer"));
+		IndexBufferRHI = RHICreateBuffer(NewNumIndices * sizeof(uint32), BUF_Dynamic | BUF_IndexBuffer | BUF_ShaderResource, sizeof(uint32), ERHIAccess::VertexOrIndexBuffer | ERHIAccess::SRVMask, CreateInfo);
 		NumIndices = NewNumIndices;
 	}
 }
 
 void FGeomCacheVertexBuffer::InitRHI()
 {
-	FRHIResourceCreateInfo CreateInfo;
-	void* BufferData = nullptr;
-	VertexBufferRHI = RHICreateAndLockVertexBuffer(SizeInBytes, BUF_Static | BUF_ShaderResource, CreateInfo, BufferData);
-	RHIUnlockBuffer(VertexBufferRHI);
+	FRHIResourceCreateInfo CreateInfo(TEXT("FGeomCacheVertexBuffer"));
+	VertexBufferRHI = RHICreateBuffer(SizeInBytes, BUF_Static | BUF_VertexBuffer | BUF_ShaderResource, 0, ERHIAccess::VertexOrIndexBuffer | ERHIAccess::SRVMask, CreateInfo);
 }
 
 void FGeomCacheVertexBuffer::UpdateRaw(const void* Data, int32 NumItems, int32 ItemSizeBytes, int32 ItemStrideBytes)
@@ -1339,19 +1333,15 @@ void FGeomCacheVertexBuffer::UpdateRaw(const void* Data, int32 NumItems, int32 I
 	int32 NewSizeInBytes = ItemSizeBytes * NumItems;
 	bool bCanMemcopy = ItemSizeBytes == ItemStrideBytes;
 
-	void* VertexBufferData = nullptr;
-
 	if (NewSizeInBytes > SizeInBytes)
 	{
 		SizeInBytes = NewSizeInBytes;
-		FRHIResourceCreateInfo CreateInfo;
-		VertexBufferRHI = RHICreateAndLockVertexBuffer(SizeInBytes, BUF_Static | BUF_ShaderResource, CreateInfo, VertexBufferData);
-	}
-	else
-	{
-		VertexBufferData = RHILockBuffer(VertexBufferRHI, 0, SizeInBytes, RLM_WriteOnly);
+		FRHIResourceCreateInfo CreateInfo(TEXT("FGeomCacheVertexBuffer"));
+		VertexBufferRHI = RHICreateBuffer(SizeInBytes, BUF_Static | BUF_VertexBuffer | BUF_ShaderResource, 0, ERHIAccess::VertexOrIndexBuffer | ERHIAccess::SRVMask, CreateInfo);
 	}
 
+	void* VertexBufferData = RHILockBuffer(VertexBufferRHI, 0, SizeInBytes, RLM_WriteOnly);
+
 	if (bCanMemcopy)
 	{
 		FMemory::Memcpy(VertexBufferData, Data, NewSizeInBytes);
@@ -1376,7 +1366,7 @@ void FGeomCacheVertexBuffer::UpdateSize(int32 NewSizeInBytes)
 	if (NewSizeInBytes > SizeInBytes)
 	{
 		SizeInBytes = NewSizeInBytes;
-		FRHIResourceCreateInfo CreateInfo;
-		VertexBufferRHI = RHICreateVertexBuffer(SizeInBytes, BUF_Static | BUF_ShaderResource, CreateInfo);
+		FRHIResourceCreateInfo CreateInfo(TEXT("FGeomCacheVertexBuffer"));
+		VertexBufferRHI = RHICreateBuffer(SizeInBytes, BUF_Static | BUF_VertexBuffer | BUF_ShaderResource, 0, ERHIAccess::VertexOrIndexBuffer | ERHIAccess::SRVMask, CreateInfo);
 	}
 }
diff --git a/Engine/Plugins/Experimental/ImagePlate/Source/ImagePlate/Private/ImagePlateComponent.cpp b/Engine/Plugins/Experimental/ImagePlate/Source/ImagePlate/Private/ImagePlateComponent.cpp
--- a/Engine/Plugins/Experimental/ImagePlate/Source/ImagePlate/Private/ImagePlateComponent.cpp
+++ b/Engine/Plugins/Experimental/ImagePlate/Source/ImagePlate/Private/ImagePlateComponent.cpp
@@ -41,12 +41,14 @@ namespace
 
 		virtual void InitRHI() override
 		{
-			FRHIResourceCreateInfo CreateInfo;
-			void* Buffer = nullptr;
-			IndexBufferRHI = RHICreateAndLockIndexBuffer(sizeof(uint16), Indices.Num() * sizeof(uint16), BUF_Static, CreateInfo, Buffer);
+			const uint32 Size = Indices.Num() * sizeof(uint16);
+
+			FRHIResourceCreateInfo CreateInfo(TEXT("FImagePlateIndexBuffer"));
+			IndexBufferRHI = RHICreateBuffer(Size, BUF_Static | BUF_IndexBuffer, sizeof(uint16), ERHIAccess::VertexOrIndexBuffer, CreateInfo);
 
 			// Copy the index data into the index buffer.		
-			FMemory::Memcpy(Buffer, Indices.GetData(), Indices.Num() * sizeof(uint16));
+			void* Buffer = RHILockBuffer(IndexBufferRHI, 0, Size, RLM_WriteOnly);
+			FMemory::Memcpy(Buffer, Indices.GetData(), Size);
 			RHIUnlockBuffer(IndexBufferRHI);
 		}
 	};
diff --git a/Engine/Plugins/Experimental/Water/Source/Runtime/Public/WaterVertexFactory.h b/Engine/Plugins/Experimental/Water/Source/Runtime/Public/WaterVertexFactory.h
--- a/Engine/Plugins/Experimental/Water/Source/Runtime/Public/WaterVertexFactory.h
+++ b/Engine/Plugins/Experimental/Water/Source/Runtime/Public/WaterVertexFactory.h
@@ -110,10 +110,9 @@ public:
 
 		NumVerts = NumVertsPerSide * NumVertsPerSide;
 
-		FRHIResourceCreateInfo CreateInfo;
-		void* BufferData = nullptr;
-		VertexBufferRHI = RHICreateAndLockVertexBuffer(sizeof(FVector4) * NumVerts, BUF_Static, CreateInfo, BufferData);
-		FVector4* DummyContents = (FVector4*)BufferData;
+		FRHIResourceCreateInfo CreateInfo(TEXT("FWaterMeshVertexBuffer"));
+		VertexBufferRHI = RHICreateBuffer(sizeof(FVector4) * NumVerts, BUF_Static | BUF_VertexBuffer, 0, ERHIAccess::VertexOrIndexBuffer, CreateInfo);
+		FVector4* DummyContents = (FVector4*)RHILockBuffer(VertexBufferRHI, 0, sizeof(FVector4) * NumVerts, RLM_WriteOnly);
 
 		for (uint32 VertY = 0; VertY < NumVertsPerSide; VertY++)
 		{
diff --git a/Engine/Plugins/FX/Niagara/Source/Niagara/Private/NiagaraDataInterfaceMeshRendererInfo.cpp b/Engine/Plugins/FX/Niagara/Source/Niagara/Private/NiagaraDataInterfaceMeshRendererInfo.cpp
--- a/Engine/Plugins/FX/Niagara/Source/Niagara/Private/NiagaraDataInterfaceMeshRendererInfo.cpp
+++ b/Engine/Plugins/FX/Niagara/Source/Niagara/Private/NiagaraDataInterfaceMeshRendererInfo.cpp
@@ -79,14 +79,13 @@ public:
 
 	virtual void InitRHI() override final
 	{
-		FRHIResourceCreateInfo CreateInfo;
-		void* BufferData = nullptr;
 		uint32 SizeByte = MeshData.Num() * sizeof(FVector) * 2;
 		
 		if (SizeByte > 0)
 		{
-			BufferMeshDataRHI = RHICreateAndLockVertexBuffer(SizeByte, BUF_Static | BUF_ShaderResource, CreateInfo, BufferData);		
-			FVector* BufferDataVector = reinterpret_cast<FVector*>(BufferData);
+			FRHIResourceCreateInfo CreateInfo(TEXT("FNDIMeshRendererInfoGPUData"));
+			BufferMeshDataRHI = RHICreateBuffer(SizeByte, BUF_Static | BUF_VertexBuffer | BUF_ShaderResource, 0, ERHIAccess::VertexOrIndexBuffer | ERHIAccess::SRVMask, CreateInfo);
+			FVector* BufferDataVector = reinterpret_cast<FVector*>(RHILockBuffer(BufferMeshDataRHI, 0, SizeByte, RLM_WriteOnly));
 			for (const auto& Mesh : MeshData)
 			{
 				*BufferDataVector++ = Mesh.MinLocalBounds;
diff --git a/Engine/Plugins/FX/Niagara/Source/Niagara/Private/NiagaraDataInterfaceSkeletalMesh.cpp b/Engine/Plugins/FX/Niagara/Source/Niagara/Private/NiagaraDataInterfaceSkeletalMesh.cpp
--- a/Engine/Plugins/FX/Niagara/Source/Niagara/Private/NiagaraDataInterfaceSkeletalMesh.cpp
+++ b/Engine/Plugins/FX/Niagara/Source/Niagara/Private/NiagaraDataInterfaceSkeletalMesh.cpp
@@ -858,10 +858,10 @@ void FSkeletalMeshGpuSpawnStaticBuffers::InitRHI()
 		const FSkeletalMeshAreaWeightedTriangleSampler& triangleSampler = SkeletalMeshSamplingLODBuiltData->AreaWeightedTriangleSampler;
 		check(TriangleCount == triangleSampler.GetNumEntries());
 
-		FRHIResourceCreateInfo CreateInfo;
-		uint32* PackedData = nullptr;
+		FRHIResourceCreateInfo CreateInfo(TEXT("FSkeletalMeshGpuSpawnStaticBuffers"));
 		uint32 SizeByte = NDISkelMeshLocal::GetProbAliasDWORDSize(TriangleCount) * sizeof(uint32);
-		BufferTriangleUniformSamplerProbAliasRHI = RHICreateAndLockVertexBuffer(SizeByte, BUF_Static | BUF_ShaderResource, CreateInfo, reinterpret_cast<void*&>(PackedData));
+		BufferTriangleUniformSamplerProbAliasRHI = RHICreateBuffer(SizeByte, BUF_Static | BUF_VertexBuffer | BUF_ShaderResource, 0, ERHIAccess::VertexOrIndexBuffer | ERHIAccess::SRVMask, CreateInfo);
+		uint32* PackedData = (uint32*)RHILockBuffer(BufferTriangleUniformSamplerProbAliasRHI, 0, SizeByte, RLM_WriteOnly);
 		NDISkelMeshLocal::PackProbAlias(PackedData, triangleSampler);
 		RHIUnlockBuffer(BufferTriangleUniformSamplerProbAliasRHI);
 		BufferTriangleUniformSamplerProbAliasSRV = RHICreateShaderResourceView(BufferTriangleUniformSamplerProbAliasRHI, sizeof(uint32), PF_R32_UINT);
@@ -902,10 +902,9 @@ void FSkeletalMeshGpuSpawnStaticBuffers::InitRHI()
 	// And these offset per section need to point to the correct matrix according to each section BoneMap.
 	// There is not section selection/culling in the interface so technically we could compute that array in the pipeline.
 	{
-		FRHIResourceCreateInfo CreateInfo;
-		void* BufferData = nullptr;
-		BufferTriangleMatricesOffsetRHI = RHICreateAndLockVertexBuffer(VertexCount * sizeof(uint32), BUF_Static | BUF_ShaderResource, CreateInfo, BufferData);
-		uint32* MatricesOffsets = (uint32*)BufferData;
+		FRHIResourceCreateInfo CreateInfo(TEXT("FSkeletalMeshGpuSpawnStaticBuffers"));
+		BufferTriangleMatricesOffsetRHI = RHICreateBuffer(VertexCount * sizeof(uint32), BUF_Static | BUF_VertexBuffer | BUF_ShaderResource, 0, ERHIAccess::VertexOrIndexBuffer | ERHIAccess::SRVMask, CreateInfo);
+		uint32* MatricesOffsets = (uint32*)RHILockBuffer(BufferTriangleMatricesOffsetRHI, 0, VertexCount * sizeof(uint32), RLM_WriteOnly);
 		uint32 AccumulatedMatrixOffset = 0;
 		for (uint32 s = 0; s < SectionCount; ++s)
 		{
diff --git a/Engine/Plugins/FX/Niagara/Source/Niagara/Private/NiagaraDataInterfaceStaticMesh.cpp b/Engine/Plugins/FX/Niagara/Source/Niagara/Private/NiagaraDataInterfaceStaticMesh.cpp
--- a/Engine/Plugins/FX/Niagara/Source/Niagara/Private/NiagaraDataInterfaceStaticMesh.cpp
+++ b/Engine/Plugins/FX/Niagara/Source/Niagara/Private/NiagaraDataInterfaceStaticMesh.cpp
@@ -182,10 +182,9 @@ void FStaticMeshGpuSpawnBuffer::InitRHI()
 	uint32 SizeByte = ValidSections.Num() * sizeof(SectionInfo);
 	if (SizeByte > 0)
 	{
-		FRHIResourceCreateInfo CreateInfo;
-		void* BufferData = nullptr;
-		BufferSectionRHI = RHICreateAndLockVertexBuffer(SizeByte, BUF_Static | BUF_ShaderResource, CreateInfo, BufferData);
-		SectionInfo* SectionInfoBuffer = (SectionInfo*)BufferData;
+		FRHIResourceCreateInfo CreateInfo(TEXT("FStaticMeshGpuSpawnBuffer"));
+		BufferSectionRHI = RHICreateBuffer(SizeByte, BUF_Static | BUF_VertexBuffer | BUF_ShaderResource, 0, ERHIAccess::VertexOrIndexBuffer | ERHIAccess::SRVMask, CreateInfo);
+		SectionInfo* SectionInfoBuffer = (SectionInfo*)RHILockBuffer(BufferSectionRHI, 0, SizeByte, RLM_WriteOnly);
 		FMemory::Memcpy(SectionInfoBuffer, ValidSections.GetData(), SizeByte);
 		RHIUnlockBuffer(BufferSectionRHI);
 		BufferSectionSRV = RHICreateShaderResourceView(BufferSectionRHI, sizeof(SectionInfo), PF_R32G32B32A32_UINT);
diff --git a/Engine/Plugins/FX/Niagara/Source/NiagaraVertexFactories/Private/NiagaraCutoutVertexBuffer.cpp b/Engine/Plugins/FX/Niagara/Source/NiagaraVertexFactories/Private/NiagaraCutoutVertexBuffer.cpp
--- a/Engine/Plugins/FX/Niagara/Source/NiagaraVertexFactories/Private/NiagaraCutoutVertexBuffer.cpp
+++ b/Engine/Plugins/FX/Niagara/Source/NiagaraVertexFactories/Private/NiagaraCutoutVertexBuffer.cpp
@@ -19,11 +19,11 @@ void FNiagaraCutoutVertexBuffer::InitRHI()
 	if (Data.Num())
 	{
 		// create a static vertex buffer
-		FRHIResourceCreateInfo CreateInfo;
-		void* BufferData = nullptr;
+		FRHIResourceCreateInfo CreateInfo(TEXT("FNiagaraCutoutVertexBuffer"));
 
 		const int32 DataSize = Data.Num() * sizeof(FVector2D);
-		VertexBufferRHI = RHICreateAndLockVertexBuffer(DataSize, BUF_Static | BUF_ShaderResource, CreateInfo, BufferData);
+		VertexBufferRHI = RHICreateBuffer(DataSize, BUF_Static | BUF_VertexBuffer | BUF_ShaderResource, 0, ERHIAccess::VertexOrIndexBuffer | ERHIAccess::SRVMask, CreateInfo);
+		void* BufferData = RHILockBuffer(VertexBufferRHI, 0, DataSize, RLM_WriteOnly);
 		FMemory::Memcpy(BufferData, Data.GetData(), DataSize);
 		RHIUnlockBuffer(VertexBufferRHI);
 		VertexBufferSRV = RHICreateShaderResourceView(VertexBufferRHI, sizeof(FVector2D), PF_G32R32F);
diff --git a/Engine/Plugins/FX/Niagara/Source/NiagaraVertexFactories/Public/NiagaraVertexFactory.h b/Engine/Plugins/FX/Niagara/Source/NiagaraVertexFactories/Public/NiagaraVertexFactory.h
--- a/Engine/Plugins/FX/Niagara/Source/NiagaraVertexFactories/Public/NiagaraVertexFactory.h
+++ b/Engine/Plugins/FX/Niagara/Source/NiagaraVertexFactories/Public/NiagaraVertexFactory.h
@@ -24,9 +24,9 @@ public:
 	virtual void InitRHI() override
 	{
 		// create a static vertex buffer
-		FRHIResourceCreateInfo CreateInfo;
-		void* BufferData = nullptr;
-		VertexBufferRHI = RHICreateAndLockVertexBuffer(sizeof(int32), BUF_Static | BUF_ShaderResource, CreateInfo, BufferData);
+		FRHIResourceCreateInfo CreateInfo(TEXT("FNiagaraNullSortedIndicesVertexBuffer"));
+		VertexBufferRHI = RHICreateBuffer(sizeof(int32), BUF_Static | BUF_VertexBuffer | BUF_ShaderResource, 0, ERHIAccess::VertexOrIndexBuffer | ERHIAccess::SRVMask, CreateInfo);
+		void* BufferData = RHILockBuffer(VertexBufferRHI, 0, sizeof(int32), RLM_WriteOnly);
 		FMemory::Memzero(BufferData, sizeof(int32));
 		RHIUnlockBuffer(VertexBufferRHI);
 
diff --git a/Engine/Plugins/Runtime/HairStrands/Source/HairStrandsCore/Private/GroomResources.cpp b/Engine/Plugins/Runtime/HairStrands/Source/HairStrandsCore/Private/GroomResources.cpp
--- a/Engine/Plugins/Runtime/HairStrands/Source/HairStrandsCore/Private/GroomResources.cpp
+++ b/Engine/Plugins/Runtime/HairStrands/Source/HairStrandsCore/Private/GroomResources.cpp
@@ -280,9 +280,9 @@ void FHairCardIndexBuffer::InitRHI()
 {
 	const uint32 DataSizeInBytes = FHairCardsIndexFormat::SizeInByte * Indices.Num();
 
-	FRHIResourceCreateInfo CreateInfo;
-	void* Buffer = nullptr;
-	IndexBufferRHI = RHICreateAndLockIndexBuffer(FHairCardsIndexFormat::SizeInByte, DataSizeInBytes, BUF_Static, CreateInfo, Buffer);
+	FRHIResourceCreateInfo CreateInfo(TEXT("FHairCardIndexBuffer"));
+	IndexBufferRHI = RHICreateBuffer(DataSizeInBytes, BUF_Static | BUF_IndexBuffer, FHairCardsIndexFormat::SizeInByte, ERHIAccess::VertexOrIndexBuffer, CreateInfo);
+	void* Buffer = RHILockBuffer(IndexBufferRHI, 0, DataSizeInBytes, RLM_WriteOnly);
 	FMemory::Memcpy(Buffer, Indices.GetData(), DataSizeInBytes);
 	RHIUnlockBuffer(IndexBufferRHI);
 }
diff --git a/Engine/Plugins/Runtime/HairStrands/Source/HairStrandsCore/Private/HairStrandsVertexFactory.cpp b/Engine/Plugins/Runtime/HairStrands/Source/HairStrandsCore/Private/HairStrandsVertexFactory.cpp
--- a/Engine/Plugins/Runtime/HairStrands/Source/HairStrandsCore/Private/HairStrandsVertexFactory.cpp
+++ b/Engine/Plugins/Runtime/HairStrands/Source/HairStrandsCore/Private/HairStrandsVertexFactory.cpp
@@ -26,11 +26,10 @@ public:
 
 	virtual void InitRHI() override
 	{
-		FRHIResourceCreateInfo CreateInfo;
-		void* BufferData = nullptr;
+		FRHIResourceCreateInfo CreateInfo(TEXT("FDummyCulledDispatchVertexIdsBuffer"));
 		uint32 NumBytes = sizeof(uint32) * 4;
-		VertexBufferRHI = RHICreateAndLockVertexBuffer(NumBytes, BUF_Static | BUF_ShaderResource, CreateInfo, BufferData);
-		uint32* DummyContents = (uint32*)BufferData;
+		VertexBufferRHI = RHICreateBuffer(NumBytes, BUF_Static | BUF_VertexBuffer | BUF_ShaderResource, 0, ERHIAccess::VertexOrIndexBuffer | ERHIAccess::SRVMask, CreateInfo);
+		uint32* DummyContents = (uint32*)RHILockBuffer(VertexBufferRHI, 0, NumBytes, RLM_WriteOnly);
 		DummyContents[0] = DummyContents[1] = DummyContents[2] = DummyContents[3] = 0;
 		RHIUnlockBuffer(VertexBufferRHI);
 
diff --git a/Engine/Plugins/Runtime/WindowsMixedReality/Source/HoloLensAR/Private/HoloLensCameraImageTexture.cpp b/Engine/Plugins/Runtime/WindowsMixedReality/Source/HoloLensAR/Private/HoloLensCameraImageTexture.cpp
--- a/Engine/Plugins/Runtime/WindowsMixedReality/Source/HoloLensAR/Private/HoloLensCameraImageTexture.cpp
+++ b/Engine/Plugins/Runtime/WindowsMixedReality/Source/HoloLensAR/Private/HoloLensCameraImageTexture.cpp
@@ -87,10 +87,9 @@ public:
 
 	virtual void InitRHI() override
 	{
-		FRHIResourceCreateInfo CreateInfo;
-		void* BufferData = nullptr;
-		VertexBufferRHI = RHICreateAndLockVertexBuffer(sizeof(FVector4) * 4, BUF_Static, CreateInfo, BufferData);
-		FVector4* DummyContents = (FVector4*)BufferData;
+		FRHIResourceCreateInfo CreateInfo(TEXT("FDummyVertexBuffer"));
+		VertexBufferRHI = RHICreateBuffer(sizeof(FVector4) * 4, BUF_Static | BUF_VertexBuffer, 0, ERHIAccess::VertexOrIndexBuffer, CreateInfo);
+		FVector4* DummyContents = (FVector4*)RHILockBuffer(VertexBufferRHI, 0, sizeof(FVector4) * 4, RLM_WriteOnly);
 		DummyContents[0] = FVector4(0.f, 0.f, 0.f, 0.f);
 		DummyContents[1] = FVector4(1.f, 0.f, 0.f, 0.f);
 		DummyContents[2] = FVector4(0.f, 1.f, 0.f, 0.f);
diff --git a/Engine/Source/Runtime/Engine/Private/GeometryCacheVertexFactory.cpp b/Engine/Source/Runtime/Engine/Private/GeometryCacheVertexFactory.cpp
--- a/Engine/Source/Runtime/Engine/Private/GeometryCacheVertexFactory.cpp
+++ b/Engine/Source/Runtime/Engine/Private/GeometryCacheVertexFactory.cpp
@@ -109,10 +109,9 @@ public:
 
 	virtual void InitRHI() override
 	{
-		FRHIResourceCreateInfo CreateInfo;
-		void* BufferData = nullptr;
-		VertexBufferRHI = RHICreateAndLockVertexBuffer(sizeof(FVector4) * 2, BUF_Static | BUF_ShaderResource, CreateInfo, BufferData);
-		FVector4* DummyContents = (FVector4*)BufferData;
+		FRHIResourceCreateInfo CreateInfo(TEXT("DefaultGeometryCacheVertexBuffer"));
+		VertexBufferRHI = RHICreateBuffer(sizeof(FVector4) * 2, BUF_Static | BUF_VertexBuffer | BUF_ShaderResource, 0, ERHIAccess::VertexOrIndexBuffer | ERHIAccess::SRVMask, CreateInfo);
+		FVector4* DummyContents = (FVector4*)RHILockBuffer(VertexBufferRHI, 0, sizeof(FVector4) * 2, RLM_WriteOnly);
 		DummyContents[0] = FVector4(0.0f, 0.0f, 0.0f, 0.0f);
 		DummyContents[1] = FVector4(1.0f, 1.0f, 1.0f, 1.0f);
 		RHIUnlockBuffer(VertexBufferRHI);
@@ -135,10 +134,9 @@ public:
 
 	virtual void InitRHI() override
 	{
-		FRHIResourceCreateInfo CreateInfo;
-		void* BufferData = nullptr;
-		VertexBufferRHI = RHICreateAndLockVertexBuffer(sizeof(FVector4) * 2, BUF_Static | BUF_ShaderResource, CreateInfo, BufferData);
-		FVector4* DummyContents = (FVector4*)BufferData;
+		FRHIResourceCreateInfo CreateInfo(TEXT("DummyTangentBuffer"));
+		VertexBufferRHI = RHICreateBuffer(sizeof(FVector4) * 2, BUF_Static | BUF_VertexBuffer | BUF_ShaderResource, 0, ERHIAccess::VertexOrIndexBuffer | ERHIAccess::SRVMask, CreateInfo);
+		FVector4* DummyContents = (FVector4*)RHILockBuffer(VertexBufferRHI, 0, sizeof(FVector4) * 2, RLM_WriteOnly);
 		DummyContents[0] = FVector4(0.0f, 0.0f, 0.0f, 0.0f);
 		DummyContents[1] = FVector4(1.0f, 1.0f, 1.0f, 1.0f);
 		RHIUnlockBuffer(VertexBufferRHI);
diff --git a/Engine/Source/Runtime/Engine/Private/Particles/ParticleCurveTexture.cpp b/Engine/Source/Runtime/Engine/Private/Particles/ParticleCurveTexture.cpp
--- a/Engine/Source/Runtime/Engine/Private/Particles/ParticleCurveTexture.cpp
+++ b/Engine/Source/Runtime/Engine/Private/Particles/ParticleCurveTexture.cpp
@@ -227,10 +227,9 @@ static void InjectCurves(
 
 		// get a buffer for all curve textures at once, and copy curve data over
 		//
-		FRHIResourceCreateInfo CreateInfo;
-		void* ScratchData = nullptr;
-		FBufferRHIRef ScratchVertexBufferRHI = RHICreateAndLockVertexBuffer(TotalSamples * sizeof(FColor), BUF_Volatile , CreateInfo, ScratchData);
-		FColor* RESTRICT DestSamples = (FColor*)ScratchData;
+		FRHIResourceCreateInfo CreateInfo(TEXT("ScratchVertexBuffer"));
+		FBufferRHIRef ScratchVertexBufferRHI = RHICreateBuffer(TotalSamples * sizeof(FColor), BUF_Volatile | BUF_VertexBuffer, 0, ERHIAccess::VertexOrIndexBuffer, CreateInfo);
+		FColor* RESTRICT DestSamples = (FColor*)RHILockBuffer(ScratchVertexBufferRHI, 0, TotalSamples * sizeof(FColor), RLM_WriteOnly);
 
 		int32 CurrOffset = 0;
 
diff --git a/Engine/Source/Runtime/Engine/Private/Particles/ParticleResources.cpp b/Engine/Source/Runtime/Engine/Private/Particles/ParticleResources.cpp
--- a/Engine/Source/Runtime/Engine/Private/Particles/ParticleResources.cpp
+++ b/Engine/Source/Runtime/Engine/Private/Particles/ParticleResources.cpp
@@ -16,10 +16,9 @@ const int32 GParticleScratchVertexBufferSize = 64 * (1 << 10); // 64KB
 void FParticleTexCoordVertexBuffer::InitRHI()
 {
 	const uint32 Size = sizeof(FVector2D) * 4 * MAX_PARTICLES_PER_INSTANCE;
-	FRHIResourceCreateInfo CreateInfo;
-	void* BufferData = nullptr;
-	VertexBufferRHI = RHICreateAndLockVertexBuffer(Size, BUF_Static, CreateInfo, BufferData);
-	FVector2D* Vertices = (FVector2D*)BufferData;
+	FRHIResourceCreateInfo CreateInfo(TEXT("FParticleTexCoordVertexBuffer"));
+	VertexBufferRHI = RHICreateBuffer(Size, BUF_Static | BUF_VertexBuffer, 0, ERHIAccess::VertexOrIndexBuffer, CreateInfo);
+	FVector2D* Vertices = (FVector2D*)RHILockBuffer(VertexBufferRHI, 0, Size, RLM_WriteOnly);
 	for (uint32 SpriteIndex = 0; SpriteIndex < MAX_PARTICLES_PER_INSTANCE; ++SpriteIndex)
 	{
 		Vertices[SpriteIndex*4 + 0] = FVector2D(0.0f, 0.0f);
@@ -39,10 +38,9 @@ TGlobalResource<FParticleTexCoordVertexBuffer> GParticleTexCoordVertexBuffer;
 void FParticleEightTexCoordVertexBuffer::InitRHI()
 {
 	const uint32 Size = sizeof(FVector2D) * 8 * MAX_PARTICLES_PER_INSTANCE;
-	FRHIResourceCreateInfo CreateInfo;
-	void* BufferData = nullptr;
-	VertexBufferRHI = RHICreateAndLockVertexBuffer(Size, BUF_Static, CreateInfo, BufferData);
-	FVector2D* Vertices = (FVector2D*)BufferData;
+	FRHIResourceCreateInfo CreateInfo(TEXT("FParticleEightTexCoordVertexBuffer"));
+	VertexBufferRHI = RHICreateBuffer(Size, BUF_Static | BUF_VertexBuffer, 0, ERHIAccess::VertexOrIndexBuffer, CreateInfo);
+	FVector2D* Vertices = (FVector2D*)RHILockBuffer(VertexBufferRHI, 0, Size, RLM_WriteOnly);
 	for (uint32 SpriteIndex = 0; SpriteIndex < MAX_PARTICLES_PER_INSTANCE; ++SpriteIndex)
 	{
 		// The contents of this buffer does not matter, whenever it is used, cutout geometry will override
@@ -72,10 +70,9 @@ void FParticleIndexBuffer::InitRHI()
 	const uint32 MaxParticles = 65536 / 4;
 	const uint32 Size = sizeof(uint16) * 6 * MaxParticles;
 	const uint32 Stride = sizeof(uint16);
-	FRHIResourceCreateInfo CreateInfo;
-	void* Buffer = nullptr;
-	IndexBufferRHI = RHICreateAndLockIndexBuffer( Stride, Size, BUF_Static, CreateInfo, Buffer);
-	uint16* Indices = (uint16*)Buffer;
+	FRHIResourceCreateInfo CreateInfo(TEXT("FParticleIndexBuffer"));
+	IndexBufferRHI = RHICreateBuffer( Size, BUF_Static | BUF_IndexBuffer, Stride, ERHIAccess::VertexOrIndexBuffer, CreateInfo );
+	uint16* Indices = (uint16*)RHILockBuffer( IndexBufferRHI, 0, Size, RLM_WriteOnly );
 	for (uint32 SpriteIndex = 0; SpriteIndex < MaxParticles; ++SpriteIndex)
 	{
 		Indices[SpriteIndex*6 + 0] = SpriteIndex*4 + 0;
@@ -102,10 +99,9 @@ void FSixTriangleParticleIndexBuffer::InitRHI()
 	const uint32 MaxParticles = 65536 / 8;
 	const uint32 Size = sizeof(uint16) * 6 * 3 * MaxParticles;
 	const uint32 Stride = sizeof(uint16);
-	FRHIResourceCreateInfo CreateInfo;
-	void* Buffer = nullptr;
-	IndexBufferRHI = RHICreateAndLockIndexBuffer( Stride, Size, BUF_Static, CreateInfo, Buffer);
-	uint16* Indices = (uint16*)Buffer;
+	FRHIResourceCreateInfo CreateInfo(TEXT("FSixTriangleParticleIndexBuffer"));
+	IndexBufferRHI = RHICreateBuffer( Size, BUF_Static | BUF_IndexBuffer, Stride, ERHIAccess::VertexOrIndexBuffer, CreateInfo );
+	uint16* Indices = (uint16*)RHILockBuffer( IndexBufferRHI, 0, Size, RLM_WriteOnly );
 	for (uint32 SpriteIndex = 0; SpriteIndex < MaxParticles; ++SpriteIndex)
 	{
 		Indices[SpriteIndex*18 + 0] = SpriteIndex*8 + 0;
diff --git a/Engine/Source/Runtime/Engine/Private/Particles/ParticleVertexFactory.cpp b/Engine/Source/Runtime/Engine/Private/Particles/ParticleVertexFactory.cpp
--- a/Engine/Source/Runtime/Engine/Private/Particles/ParticleVertexFactory.cpp
+++ b/Engine/Source/Runtime/Engine/Private/Particles/ParticleVertexFactory.cpp
@@ -22,11 +22,13 @@ public:
 	 */
 	virtual void InitRHI() override
 	{
+		const uint32 Size = sizeof(FVector2D) * 4;
+
 		// create a static vertex buffer
-		FRHIResourceCreateInfo CreateInfo;
-		void* BufferData = nullptr;
-		VertexBufferRHI = RHICreateAndLockVertexBuffer(sizeof(FVector2D) * 4, BUF_Static | BUF_ShaderResource, CreateInfo, BufferData);
-		FMemory::Memzero(BufferData, sizeof(FVector2D) * 4);
+		FRHIResourceCreateInfo CreateInfo(TEXT("FNullSubUVCutoutVertexBuffer"));
+		VertexBufferRHI = RHICreateBuffer(Size, BUF_Static | BUF_VertexBuffer | BUF_ShaderResource, 0, ERHIAccess::VertexOrIndexBuffer | ERHIAccess::SRVMask, CreateInfo);
+		void* BufferData = RHILockBuffer(VertexBufferRHI, 0, Size, RLM_WriteOnly);
+		FMemory::Memzero(BufferData, Size);
 		RHIUnlockBuffer(VertexBufferRHI);
 
 		if (GSupportsResourceView)
diff --git a/Engine/Source/Runtime/Engine/Private/Rendering/MorphTargetVertexInfoBuffers.cpp b/Engine/Source/Runtime/Engine/Private/Rendering/MorphTargetVertexInfoBuffers.cpp
--- a/Engine/Source/Runtime/Engine/Private/Rendering/MorphTargetVertexInfoBuffers.cpp
+++ b/Engine/Source/Runtime/Engine/Private/Rendering/MorphTargetVertexInfoBuffers.cpp
@@ -10,17 +10,17 @@ void FMorphTargetVertexInfoBuffers::InitRHI()
 	check(NumTotalWorkItems > 0);
 
 	{
-		FRHIResourceCreateInfo CreateInfo;
-		void* VertexIndicesVBData = nullptr;
-		VertexIndicesVB = RHICreateAndLockVertexBuffer(VertexIndices.GetAllocatedSize(), BUF_Static | BUF_ShaderResource, CreateInfo, VertexIndicesVBData);
+		FRHIResourceCreateInfo CreateInfo(TEXT("VertexIndicesVB"));
+		VertexIndicesVB = RHICreateBuffer(VertexIndices.GetAllocatedSize(), BUF_Static | BUF_VertexBuffer | BUF_ShaderResource, 0, ERHIAccess::VertexOrIndexBuffer | ERHIAccess::SRVMask, CreateInfo);
+		void* VertexIndicesVBData = RHILockBuffer(VertexIndicesVB, 0, VertexIndices.GetAllocatedSize(), RLM_WriteOnly);
 		FMemory::ParallelMemcpy(VertexIndicesVBData, VertexIndices.GetData(), VertexIndices.GetAllocatedSize(), EMemcpyCachePolicy::StoreUncached);
 		RHIUnlockBuffer(VertexIndicesVB);
 		VertexIndicesSRV = RHICreateShaderResourceView(VertexIndicesVB, 4, PF_R32_UINT);
 	}
 	{
-		FRHIResourceCreateInfo CreateInfo;
-		void* MorphDeltasVBData = nullptr;
-		MorphDeltasVB = RHICreateAndLockVertexBuffer(MorphDeltas.GetAllocatedSize(), BUF_Static | BUF_ShaderResource, CreateInfo, MorphDeltasVBData);
+		FRHIResourceCreateInfo CreateInfo(TEXT("MorphDeltasVB"));
+		MorphDeltasVB = RHICreateBuffer(MorphDeltas.GetAllocatedSize(), BUF_Static | BUF_VertexBuffer | BUF_ShaderResource, 0, ERHIAccess::VertexOrIndexBuffer | ERHIAccess::SRVMask, CreateInfo);
+		void* MorphDeltasVBData = RHILockBuffer(MorphDeltasVB, 0, MorphDeltas.GetAllocatedSize(), RLM_WriteOnly);
 		FMemory::ParallelMemcpy(MorphDeltasVBData, MorphDeltas.GetData(), MorphDeltas.GetAllocatedSize(), EMemcpyCachePolicy::StoreUncached);
 		RHIUnlockBuffer(MorphDeltasVB);
 		MorphDeltasSRV = RHICreateShaderResourceView(MorphDeltasVB, 2, PF_R16F);
diff --git a/Engine/Source/Runtime/Engine/Private/StaticMesh.cpp b/Engine/Source/Runtime/Engine/Private/StaticMesh.cpp
--- a/Engine/Source/Runtime/Engine/Private/StaticMesh.cpp
+++ b/Engine/Source/Runtime/Engine/Private/StaticMesh.cpp
@@ -218,9 +218,6 @@ void FStaticMeshSectionAreaWeightedTriangleSamplerBuffer::InitRHI()
 
 	if (Samplers && Samplers->Num() > 0)
 	{
-		FRHIResourceCreateInfo CreateInfo;
-		void* BufferData = nullptr;
-
 		// Count triangle count for all sections and required memory
 		const uint32 AllSectionCount = Samplers->Num();
 		uint32 TriangleCount = 0;
@@ -230,10 +227,11 @@ void FStaticMeshSectionAreaWeightedTriangleSamplerBuffer::InitRHI()
 		}
 		uint32 SizeByte = TriangleCount * sizeof(SectionTriangleInfo);
 
-		BufferSectionTriangleRHI = RHICreateAndLockVertexBuffer(SizeByte, BUF_Static | BUF_ShaderResource, CreateInfo, BufferData);
+		FRHIResourceCreateInfo CreateInfo(TEXT("StaticMeshSectionAreaWeightedTriangleSamplerBuffer"));
+		BufferSectionTriangleRHI = RHICreateBuffer(SizeByte, BUF_Static | BUF_VertexBuffer | BUF_ShaderResource, 0, ERHIAccess::VertexOrIndexBuffer | ERHIAccess::SRVMask, CreateInfo);
 
 		// Now compute the alias look up table for unifor; distribution for all section and all triangles
-		SectionTriangleInfo* SectionTriangleInfoBuffer = (SectionTriangleInfo*)BufferData;
+		SectionTriangleInfo* SectionTriangleInfoBuffer = (SectionTriangleInfo*)RHILockBuffer(BufferSectionTriangleRHI, 0, SizeByte, RLM_WriteOnly);
 		for (uint32 i = 0; i < AllSectionCount; ++i)
 		{
 			FStaticMeshSectionAreaWeightedTriangleSampler& sampler = (*Samplers)[i];
diff --git a/Engine/Source/Runtime/Engine/Private/VectorFieldVisualization.cpp b/Engine/Source/Runtime/Engine/Private/VectorFieldVisualization.cpp
--- a/Engine/Source/Runtime/Engine/Private/VectorFieldVisualization.cpp
+++ b/Engine/Source/Runtime/Engine/Private/VectorFieldVisualization.cpp
@@ -88,10 +88,9 @@ public:
 
 	virtual void InitRHI() override
 	{
-		FRHIResourceCreateInfo CreateInfo;
-		void* BufferData = nullptr;
-		VertexBufferRHI = RHICreateAndLockVertexBuffer(sizeof(FVector4) * 2, BUF_Static, CreateInfo, BufferData);
-		FVector4* DummyContents = (FVector4*)BufferData;
+		FRHIResourceCreateInfo CreateInfo(TEXT("FDummyVertexBuffer"));
+		VertexBufferRHI = RHICreateBuffer(sizeof(FVector4) * 2, BUF_Static | BUF_VertexBuffer, 0, ERHIAccess::VertexOrIndexBuffer, CreateInfo);
+		FVector4* DummyContents = (FVector4*)RHILockBuffer(VertexBufferRHI, 0, sizeof(FVector4) * 2, RLM_WriteOnly);
 		DummyContents[0] = FVector4(0.0f, 0.0f, 0.0f, 0.0f);
 		DummyContents[1] = FVector4(1.0f, 1.0f, 1.0f, 1.0f);
 		RHIUnlockBuffer(VertexBufferRHI);
diff --git a/Engine/Source/Runtime/Engine/Public/ParticleHelper.h b/Engine/Source/Runtime/Engine/Public/ParticleHelper.h
--- a/Engine/Source/Runtime/Engine/Public/ParticleHelper.h
+++ b/Engine/Source/Runtime/Engine/Public/ParticleHelper.h
@@ -2841,10 +2841,9 @@ public:
 	virtual void InitRHI() override
 	{
 		// create a static vertex buffer
-		FRHIResourceCreateInfo CreateInfo;
-		void* BufferData = nullptr;
-		VertexBufferRHI = RHICreateAndLockVertexBuffer(sizeof(FParticleVertexDynamicParameter), BUF_Static | BUF_ZeroStride, CreateInfo, BufferData);
-		FParticleVertexDynamicParameter* Vertices = (FParticleVertexDynamicParameter*)BufferData;
+		FRHIResourceCreateInfo CreateInfo(TEXT("FNullDynamicParameterVertexBuffer"));
+		VertexBufferRHI = RHICreateBuffer(sizeof(FParticleVertexDynamicParameter), BUF_Static | BUF_VertexBuffer | BUF_ZeroStride, 0, ERHIAccess::VertexOrIndexBuffer, CreateInfo);
+		FParticleVertexDynamicParameter* Vertices = (FParticleVertexDynamicParameter*)RHILockBuffer(VertexBufferRHI, 0, sizeof(FParticleVertexDynamicParameter), RLM_WriteOnly);
 		Vertices[0].DynamicValue[0] = Vertices[0].DynamicValue[1] = Vertices[0].DynamicValue[2] = Vertices[0].DynamicValue[3] = 1.0f;
 		RHIUnlockBuffer(VertexBufferRHI);
 	}
diff --git a/Engine/Source/Runtime/Landscape/Private/LandscapeRender.cpp b/Engine/Source/Runtime/Landscape/Private/LandscapeRender.cpp
--- a/Engine/Source/Runtime/Landscape/Private/LandscapeRender.cpp
+++ b/Engine/Source/Runtime/Landscape/Private/LandscapeRender.cpp
@@ -2829,10 +2829,9 @@ void FLandscapeVertexBuffer::InitRHI()
 	SCOPED_LOADTIMER(FLandscapeVertexBuffer_InitRHI);
 
 	// create a static vertex buffer
-	FRHIResourceCreateInfo CreateInfo;
-	void* BufferData = nullptr;
-	VertexBufferRHI = RHICreateAndLockVertexBuffer(NumVertices * sizeof(FLandscapeVertex), BUF_Static, CreateInfo, BufferData);
-	FLandscapeVertex* Vertex = (FLandscapeVertex*)BufferData;
+	FRHIResourceCreateInfo CreateInfo(TEXT("FLandscapeVertexBuffer"));
+	VertexBufferRHI = RHICreateBuffer(NumVertices * sizeof(FLandscapeVertex), BUF_Static | BUF_VertexBuffer, 0, ERHIAccess::VertexOrIndexBuffer, CreateInfo);
+	FLandscapeVertex* Vertex = (FLandscapeVertex*)RHILockBuffer(VertexBufferRHI, 0, NumVertices * sizeof(FLandscapeVertex), RLM_WriteOnly);
 	int32 VertexIndex = 0;
 	for (int32 SubY = 0; SubY < NumSubsections; SubY++)
 	{
diff --git a/Engine/Source/Runtime/Landscape/Private/LandscapeRenderMobile.cpp b/Engine/Source/Runtime/Landscape/Private/LandscapeRenderMobile.cpp
--- a/Engine/Source/Runtime/Landscape/Private/LandscapeRenderMobile.cpp
+++ b/Engine/Source/Runtime/Landscape/Private/LandscapeRenderMobile.cpp
@@ -211,11 +211,11 @@ void FLandscapeVertexBufferMobile::UpdateMemoryStat(int32 Delta)
 void FLandscapeVertexBufferMobile::InitRHI()
 {
 	// create a static vertex buffer
-	FRHIResourceCreateInfo CreateInfo;
-	void* VertexDataPtr = nullptr;
-	VertexBufferRHI = RHICreateAndLockVertexBuffer(VertexData.Num(), BUF_Static, CreateInfo, VertexDataPtr);
+	FRHIResourceCreateInfo CreateInfo(TEXT("FLandscapeVertexBufferMobile"));
+	VertexBufferRHI = RHICreateBuffer(VertexData.Num(), BUF_Static | BUF_VertexBuffer, 0, ERHIAccess::VertexOrIndexBuffer, CreateInfo);
 
 	// Copy stored platform data and free CPU copy
+	void* VertexDataPtr = RHILockBuffer(VertexBufferRHI, 0, VertexData.Num(), RLM_WriteOnly);
 	FMemory::Memcpy(VertexDataPtr, VertexData.GetData(), VertexData.Num());
 	VertexData.Empty();
 
diff --git a/Engine/Source/Runtime/MRMesh/Private/MRMeshComponent.cpp b/Engine/Source/Runtime/MRMesh/Private/MRMeshComponent.cpp
--- a/Engine/Source/Runtime/MRMesh/Private/MRMeshComponent.cpp
+++ b/Engine/Source/Runtime/MRMesh/Private/MRMeshComponent.cpp
@@ -105,12 +105,14 @@ public:
 	{
 		NumIndices = Indices.Num();
 
-		FRHIResourceCreateInfo CreateInfo;
-		void* Buffer = nullptr;
-		IndexBufferRHI = RHICreateAndLockIndexBuffer(sizeof(int32), Indices.Num() * sizeof(int32), BUF_Static, CreateInfo, Buffer);
+		const uint32 Size = Indices.Num() * sizeof(uint32);
+
+		FRHIResourceCreateInfo CreateInfo(TEXT("FMRMeshIndexBuffer"));
+		IndexBufferRHI = RHICreateBuffer(Size, BUF_Static | BUF_IndexBuffer, sizeof(uint32), ERHIAccess::VertexOrIndexBuffer, CreateInfo);
 
 		// Write the indices to the index buffer.
-		FMemory::Memcpy(Buffer, Indices.GetData(), Indices.Num() * sizeof(int32));
+		void* Buffer = RHILockBuffer(IndexBufferRHI, 0, Size, RLM_WriteOnly);
+		FMemory::Memcpy(Buffer, Indices.GetData(), Size);
 		RHIUnlockBuffer(IndexBufferRHI);
 	}
 
@@ -118,12 +120,14 @@ public:
 	{
 		NumIndices = Indices.Num();
 
-		FRHIResourceCreateInfo CreateInfo;
-		void* Buffer = nullptr;
-		IndexBufferRHI = RHICreateAndLockIndexBuffer(sizeof(uint16), Indices.Num() * sizeof(uint16), BUF_Static, CreateInfo, Buffer);
+		const uint32 Size = Indices.Num() * sizeof(uint16);
+
+		FRHIResourceCreateInfo CreateInfo(TEXT("FMRMeshIndexBuffer"));
+		IndexBufferRHI = RHICreateBuffer(Size, BUF_Static | BUF_IndexBuffer, sizeof(uint16), ERHIAccess::VertexOrIndexBuffer, CreateInfo);
 
 		// Write the indices to the index buffer.
-		FMemory::Memcpy(Buffer, Indices.GetData(), Indices.Num() * sizeof(uint16));
+		void* Buffer = RHILockBuffer(IndexBufferRHI, 0, Size, RLM_WriteOnly);
+		FMemory::Memcpy(Buffer, Indices.GetData(), Size);
 		RHIUnlockBuffer(IndexBufferRHI);
 	}
 };
diff --git a/Engine/Source/Runtime/RenderCore/Private/VertexFactory.cpp b/Engine/Source/Runtime/RenderCore/Private/VertexFactory.cpp
--- a/Engine/Source/Runtime/RenderCore/Private/VertexFactory.cpp
+++ b/Engine/Source/Runtime/RenderCore/Private/VertexFactory.cpp
@@ -327,9 +327,8 @@ void FPrimitiveIdDummyBuffer::InitRHI()
 	// create a static vertex buffer
 	FRHIResourceCreateInfo CreateInfo(TEXT("FPrimitiveIdDummyBuffer"));
 
-	void* LockedData = nullptr;
-	VertexBufferRHI = RHICreateAndLockVertexBuffer(sizeof(uint32), BUF_Static | BUF_ShaderResource, CreateInfo, LockedData);
-	uint32* Vertices = (uint32*)LockedData;
+	VertexBufferRHI = RHICreateBuffer(sizeof(uint32), BUF_Static | BUF_VertexBuffer | BUF_ShaderResource, 0, ERHIAccess::VertexOrIndexBuffer | ERHIAccess::SRVMask, CreateInfo);
+	uint32* Vertices = (uint32*)RHILockBuffer(VertexBufferRHI, 0, sizeof(uint32), RLM_WriteOnly);
 	Vertices[0] = 0;
 	RHIUnlockBuffer(VertexBufferRHI);
 	VertexBufferSRV = RHICreateShaderResourceView(VertexBufferRHI, sizeof(uint32), PF_R32_UINT);
diff --git a/Engine/Source/Runtime/RenderCore/Public/RenderGraphResources.h b/Engine/Source/Runtime/RenderCore/Public/RenderGraphResources.h
--- a/Engine/Source/Runtime/RenderCore/Public/RenderGraphResources.h
+++ b/Engine/Source/Runtime/RenderCore/Public/RenderGraphResources.h
@@ -1191,7 +1191,7 @@ public:
 		return GetRHI();
 	}
 
-	/** Returns the buffer to use for RHI calls, eg RHILockVertexBuffer. */
+	/** Returns the buffer to use for RHI calls, eg RHILockBuffer. */
 	UE_DEPRECATED(5.0, "Buffers types have been consolidated; use GetRHI() instead.")
 	FORCEINLINE FRHIBuffer* GetRHIVertexBuffer() const
 	{
diff --git a/Engine/Source/Runtime/RenderCore/Public/RenderResource.h b/Engine/Source/Runtime/RenderCore/Public/RenderResource.h
--- a/Engine/Source/Runtime/RenderCore/Public/RenderResource.h
+++ b/Engine/Source/Runtime/RenderCore/Public/RenderResource.h
@@ -543,9 +543,8 @@ public:
 		// create a static vertex buffer
 		FRHIResourceCreateInfo CreateInfo(TEXT("FNullColorVertexBuffer"));
 
-		void* LockedData = nullptr;
-		VertexBufferRHI = RHICreateAndLockVertexBuffer(sizeof(uint32) * 4, BUF_Static | BUF_ZeroStride | BUF_ShaderResource, CreateInfo, LockedData);
-		uint32* Vertices = (uint32*)LockedData;
+		VertexBufferRHI = RHICreateBuffer(sizeof(uint32) * 4, BUF_Static | BUF_ZeroStride | BUF_VertexBuffer | BUF_ShaderResource, 0, ERHIAccess::VertexOrIndexBuffer | ERHIAccess::SRVMask, CreateInfo);
+		uint32* Vertices = (uint32*)RHILockBuffer(VertexBufferRHI, 0, sizeof(uint32) * 4, RLM_WriteOnly);
 		Vertices[0] = FColor(255, 255, 255, 255).DWColor();
 		Vertices[1] = FColor(255, 255, 255, 255).DWColor();
 		Vertices[2] = FColor(255, 255, 255, 255).DWColor();
@@ -578,13 +577,10 @@ public:
 	virtual void InitRHI() override
 	{
 		// create a static vertex buffer
-		FRHIResourceCreateInfo CreateInfo;
-
-		void* LockedData = nullptr;
-		VertexBufferRHI = RHICreateAndLockVertexBuffer(sizeof(float) * 3, BUF_Static | BUF_ZeroStride | BUF_ShaderResource, CreateInfo, LockedData);
-
-		*reinterpret_cast<FVector*>(LockedData) = FVector(0.0f);
-
+		FRHIResourceCreateInfo CreateInfo(TEXT("FNullVertexBuffer"));
+		VertexBufferRHI = RHICreateBuffer(sizeof(FVector), BUF_Static | BUF_ZeroStride | BUF_VertexBuffer | BUF_ShaderResource, 0, ERHIAccess::VertexOrIndexBuffer | ERHIAccess::SRVMask, CreateInfo);
+		FVector* LockedData = (FVector*)RHILockBuffer(VertexBufferRHI, 0, sizeof(FVector), RLM_WriteOnly);
+		*LockedData = FVector(0.0f);
 		RHIUnlockBuffer(VertexBufferRHI);
 
 		VertexBufferSRV = RHICreateShaderResourceView(VertexBufferRHI, sizeof(FColor), PF_R8G8B8A8);
diff --git a/Engine/Source/Runtime/Renderer/Private/DistanceFieldObjectCulling.cpp b/Engine/Source/Runtime/Renderer/Private/DistanceFieldObjectCulling.cpp
--- a/Engine/Source/Runtime/Renderer/Private/DistanceFieldObjectCulling.cpp
+++ b/Engine/Source/Runtime/Renderer/Private/DistanceFieldObjectCulling.cpp
@@ -50,10 +50,9 @@ public:
 	{
 		// Used as a non-indexed triangle list, so 3 vertices per triangle
 		const uint32 Size = 3 * NumSections * sizeof(FScreenVertex);
-		FRHIResourceCreateInfo CreateInfo;
-		void* Buffer = nullptr;
-		VertexBufferRHI = RHICreateAndLockVertexBuffer(Size, BUF_Static, CreateInfo, Buffer);
-		FScreenVertex* DestVertex = (FScreenVertex*)Buffer;
+		FRHIResourceCreateInfo CreateInfo(TEXT("FCircleVertexBuffer"));
+		VertexBufferRHI = RHICreateBuffer(Size, BUF_Static | BUF_VertexBuffer, 0, ERHIAccess::VertexOrIndexBuffer, CreateInfo);
+		FScreenVertex* DestVertex = (FScreenVertex*)RHILockBuffer(VertexBufferRHI, 0, Size, RLM_WriteOnly);
 
 		const float RadiansPerRingSegment = PI / (float)NumSections;
 
diff --git a/Engine/Source/Runtime/Renderer/Private/HeightfieldLighting.cpp b/Engine/Source/Runtime/Renderer/Private/HeightfieldLighting.cpp
--- a/Engine/Source/Runtime/Renderer/Private/HeightfieldLighting.cpp
+++ b/Engine/Source/Runtime/Renderer/Private/HeightfieldLighting.cpp
@@ -338,11 +338,9 @@ public:
 	virtual void InitRHI() override
 	{
 		const uint32 Size = 6 * sizeof(FScreenVertex);
-		FRHIResourceCreateInfo CreateInfo;
-
-		void* Buffer = nullptr;
-		VertexBufferRHI = RHICreateAndLockVertexBuffer(Size, BUF_Static, CreateInfo, Buffer);
-		FScreenVertex* DestVertex = (FScreenVertex*)Buffer;
+		FRHIResourceCreateInfo CreateInfo(TEXT("FQuadVertexBuffer"));
+		VertexBufferRHI = RHICreateBuffer(Size, BUF_Static | BUF_VertexBuffer, 0, ERHIAccess::VertexOrIndexBuffer, CreateInfo);
+		FScreenVertex* DestVertex = (FScreenVertex*)RHILockBuffer(VertexBufferRHI, 0, Size, RLM_WriteOnly);
 
 		DestVertex[0].Position = FVector2D(1, 1);
 		DestVertex[0].UV = FVector2D(1, 1);
diff --git a/Engine/Source/Runtime/Renderer/Private/RendererPrivateUtils.h b/Engine/Source/Runtime/Renderer/Private/RendererPrivateUtils.h
--- a/Engine/Source/Runtime/Renderer/Private/RendererPrivateUtils.h
+++ b/Engine/Source/Runtime/Renderer/Private/RendererPrivateUtils.h
@@ -14,9 +14,8 @@ public:
 	{
 		const uint32 Size = sizeof(FVector2D) * 4 * NumTileQuadsInBuffer;
 		FRHIResourceCreateInfo CreateInfo;
-		void* BufferData = nullptr;
-		VertexBufferRHI = RHICreateAndLockVertexBuffer(Size, BUF_Static, CreateInfo, BufferData);
-		FVector2D* Vertices = (FVector2D*)BufferData;
+		VertexBufferRHI = RHICreateBuffer(Size, BUF_Static | BUF_VertexBuffer, 0, ERHIAccess::VertexOrIndexBuffer, CreateInfo);
+		FVector2D* Vertices = (FVector2D*)RHILockBuffer(VertexBufferRHI, 0, Size, RLM_WriteOnly);
 		for (uint32 SpriteIndex = 0; SpriteIndex < NumTileQuadsInBuffer; ++SpriteIndex)
 		{
 			Vertices[SpriteIndex * 4 + 0] = FVector2D(0.0f, 0.0f);
@@ -43,10 +42,9 @@ public:
 	{
 		const uint32 Size = sizeof(uint16) * 6 * NumTileQuadsInBuffer;
 		const uint32 Stride = sizeof(uint16);
-		FRHIResourceCreateInfo CreateInfo;
-		void* Buffer = nullptr;
-		IndexBufferRHI = RHICreateAndLockIndexBuffer(Stride, Size, BUF_Static, CreateInfo, Buffer);
-		uint16* Indices = (uint16*)Buffer;
+		FRHIResourceCreateInfo CreateInfo(TEXT("FTileIndexBuffer"));
+		IndexBufferRHI = RHICreateBuffer(Size, BUF_Static | BUF_IndexBuffer, Stride, ERHIAccess::VertexOrIndexBuffer, CreateInfo);
+		uint16* Indices = (uint16*)RHILockBuffer(IndexBufferRHI, 0, Size, RLM_WriteOnly);
 		for (uint32 SpriteIndex = 0; SpriteIndex < NumTileQuadsInBuffer; ++SpriteIndex)
 		{
 			Indices[SpriteIndex * 6 + 0] = SpriteIndex * 4 + 0;
diff --git a/Engine/Source/Runtime/Renderer/Private/SceneOcclusion.cpp b/Engine/Source/Runtime/Renderer/Private/SceneOcclusion.cpp
--- a/Engine/Source/Runtime/Renderer/Private/SceneOcclusion.cpp
+++ b/Engine/Source/Runtime/Renderer/Private/SceneOcclusion.cpp
@@ -338,11 +338,9 @@ public:
 		const uint32 Stride = sizeof(uint16);
 		const uint32 SizeInBytes = MaxBatchedPrimitives * NUM_CUBE_VERTICES * Stride;
 
-		FRHIResourceCreateInfo CreateInfo;
-
-		void* BufferData;
-		IndexBufferRHI = RHICreateAndLockIndexBuffer(Stride, SizeInBytes, BUF_Static, CreateInfo, BufferData);
-		uint16* RESTRICT Indices = (uint16*)BufferData;
+		FRHIResourceCreateInfo CreateInfo(TEXT("FOcclusionQueryIndexBuffer"));
+		IndexBufferRHI = RHICreateBuffer(SizeInBytes, BUF_Static | BUF_IndexBuffer, Stride, ERHIAccess::VertexOrIndexBuffer, CreateInfo);
+		uint16* RESTRICT Indices = (uint16*) RHILockBuffer(IndexBufferRHI, 0, SizeInBytes, RLM_WriteOnly);
 
 		for(uint32 PrimitiveIndex = 0;PrimitiveIndex < MaxBatchedPrimitives;PrimitiveIndex++)
 		{
diff --git a/Engine/Source/Runtime/Renderer/Private/SceneRendering.cpp b/Engine/Source/Runtime/Renderer/Private/SceneRendering.cpp
--- a/Engine/Source/Runtime/Renderer/Private/SceneRendering.cpp
+++ b/Engine/Source/Runtime/Renderer/Private/SceneRendering.cpp
@@ -4317,9 +4317,9 @@ public:
 	{
 		const int32 NumDummyVerts = 3;
 		const uint32 Size = sizeof(FVector4) * NumDummyVerts;
-		FRHIResourceCreateInfo CreateInfo;
-		void* BufferData = nullptr;
-		VertexBufferRHI = RHICreateAndLockVertexBuffer(Size, BUF_Static, CreateInfo, BufferData);
+		FRHIResourceCreateInfo CreateInfo(TEXT("FDummySceneColorResolveBuffer"));
+		VertexBufferRHI = RHICreateBuffer(Size, BUF_Static | BUF_VertexBuffer, 0, ERHIAccess::VertexOrIndexBuffer, CreateInfo);
+		void* BufferData = RHILockBuffer(VertexBufferRHI, 0, Size, RLM_WriteOnly);
 		FMemory::Memset(BufferData, 0, Size);		
 		RHIUnlockBuffer(VertexBufferRHI);		
 	}
diff --git a/Engine/Source/Runtime/Renderer/Private/SpriteIndexBuffer.h b/Engine/Source/Runtime/Renderer/Private/SpriteIndexBuffer.h
--- a/Engine/Source/Runtime/Renderer/Private/SpriteIndexBuffer.h
+++ b/Engine/Source/Runtime/Renderer/Private/SpriteIndexBuffer.h
@@ -10,10 +10,9 @@ public:
 	{
 		const uint32 Size = sizeof(uint16) * 6 * NumSprites;
 		const uint32 Stride = sizeof(uint16);
-		FRHIResourceCreateInfo CreateInfo;
-		void* Buffer = nullptr;
-		IndexBufferRHI = RHICreateAndLockIndexBuffer( Stride, Size, BUF_Static, CreateInfo, Buffer );
-		uint16* Indices = (uint16*)Buffer;
+		FRHIResourceCreateInfo CreateInfo(TEXT("FSpriteIndexBuffer"));
+		IndexBufferRHI = RHICreateBuffer( Size, BUF_Static | BUF_IndexBuffer, Stride, ERHIAccess::VertexOrIndexBuffer, CreateInfo );
+		uint16* Indices = (uint16*)RHILockBuffer( IndexBufferRHI, 0, Size, RLM_WriteOnly );
 		for (uint32 SpriteIndex = 0; SpriteIndex < NumSprites; ++SpriteIndex)
 		{
 			Indices[SpriteIndex*6 + 0] = SpriteIndex*4 + 0;
diff --git a/Engine/Source/Runtime/Renderer/Private/VolumetricFog.cpp b/Engine/Source/Runtime/Renderer/Private/VolumetricFog.cpp
--- a/Engine/Source/Runtime/Renderer/Private/VolumetricFog.cpp
+++ b/Engine/Source/Runtime/Renderer/Private/VolumetricFog.cpp
@@ -521,10 +521,9 @@ public:
 	{
 		const int32 NumTriangles = NumVertices - 2;
 		const uint32 Size = NumVertices * sizeof(FScreenVertex);
-		FRHIResourceCreateInfo CreateInfo;
-		void* Buffer = nullptr;
-		VertexBufferRHI = RHICreateAndLockVertexBuffer(Size, BUF_Static, CreateInfo, Buffer);
-		FScreenVertex* DestVertex = (FScreenVertex*)Buffer;
+		FRHIResourceCreateInfo CreateInfo(TEXT("FCircleRasterizeVertexBuffer"));
+		VertexBufferRHI = RHICreateBuffer(Size, BUF_Static | BUF_VertexBuffer, 0, ERHIAccess::VertexOrIndexBuffer, CreateInfo);
+		FScreenVertex* DestVertex = (FScreenVertex*)RHILockBuffer(VertexBufferRHI, 0, Size, RLM_WriteOnly);
 
 		const int32 NumRings = NumVertices;
 		const float RadiansPerRingSegment = PI / (float)NumRings;
diff --git a/Engine/Source/Runtime/Renderer/Public/VolumeRendering.h b/Engine/Source/Runtime/Renderer/Public/VolumeRendering.h
--- a/Engine/Source/Runtime/Renderer/Public/VolumeRendering.h
+++ b/Engine/Source/Runtime/Renderer/Public/VolumeRendering.h
@@ -127,10 +127,9 @@ public:
 	{
 		// Used as a non-indexed triangle strip, so 4 vertices per quad
 		const uint32 Size = 4 * sizeof(FScreenVertex);
-		FRHIResourceCreateInfo CreateInfo;
-		void* Buffer = nullptr;
-		VertexBufferRHI = RHICreateAndLockVertexBuffer(Size, BUF_Static, CreateInfo, Buffer);		
-		FScreenVertex* DestVertex = (FScreenVertex*)Buffer;
+		FRHIResourceCreateInfo CreateInfo(TEXT("FVolumeRasterizeVertexBuffer"));
+		VertexBufferRHI = RHICreateBuffer(Size, BUF_Static | BUF_VertexBuffer, 0, ERHIAccess::VertexOrIndexBuffer, CreateInfo);
+		FScreenVertex* DestVertex = (FScreenVertex*)RHILockBuffer(VertexBufferRHI, 0, Size, RLM_WriteOnly);
 
 		// Setup a full - render target quad
 		// A viewport and UVScaleBias will be used to implement rendering to a sub region
