GeometryCollection: Define HUFFMAN_USE_UNALIGNED_READ based on PLATFORM_SUPPORTS_UNALIGNED_LOADS and PLATFORM_LITTLE_ENDIAN instead of relying on explicitly named platform defines or plugin extensions

diff --git a/Engine/Plugins/Experimental/GeometryCache/Source/GeometryCache/Private/HuffmanBitStream.h b/Engine/Plugins/Experimental/GeometryCache/Source/GeometryCache/Private/HuffmanBitStream.h
--- a/Engine/Plugins/Experimental/GeometryCache/Source/GeometryCache/Private/HuffmanBitStream.h
+++ b/Engine/Plugins/Experimental/GeometryCache/Source/GeometryCache/Private/HuffmanBitStream.h
@@ -4,11 +4,8 @@
 #include "CoreMinimal.h"
 #include "Containers/Array.h"
 
-#if !defined(HUFFMAN_USE_UNALIGNED_READ)
-	#define HUFFMAN_USE_UNALIGNED_READ (PLATFORM_WINDOWS | PLATFORM_MAC | PLATFORM_PS4)	// Little-endian platforms that support fast unaligned reads
-#endif
-
-#define MINIMUM_BITS_AFTER_REFILL 56															// Minimum number of bits guaranteed to be available in the internal buffer after a buffer refill.
+#define HUFFMAN_USE_UNALIGNED_READ	(PLATFORM_SUPPORTS_UNALIGNED_LOADS && PLATFORM_LITTLE_ENDIAN)
+#define MINIMUM_BITS_AFTER_REFILL	56																// Minimum number of bits guaranteed to be available in the internal buffer after a buffer refill.
 
 /**
 A bit stream writer class for use with the Huffman coding.
