Renamed UE4 to UE/UnrealEngine.

diff --git a/Engine/Plugins/Experimental/AlembicImporter/Source/AlembicLibrary/Private/AbcImportUtilities.h b/Engine/Plugins/Experimental/AlembicImporter/Source/AlembicLibrary/Private/AbcImportUtilities.h
--- a/Engine/Plugins/Experimental/AlembicImporter/Source/AlembicLibrary/Private/AbcImportUtilities.h
+++ b/Engine/Plugins/Experimental/AlembicImporter/Source/AlembicLibrary/Private/AbcImportUtilities.h
@@ -60,7 +60,7 @@ namespace AbcImporterUtilities
 	}
 
 	/**
-	* ConvertAlembicMatrix, converts Abc(Alembic) matrix to UE4 matrix format
+	* ConvertAlembicMatrix, converts Abc(Alembic) matrix to UnrealEngine matrix format
 	*
 	* @param AbcMatrix - Alembic style matrix
 	* @return FMatrix
diff --git a/Engine/Plugins/Experimental/GeometryCache/Source/GeometryCache/Private/NormalCompression.cpp b/Engine/Plugins/Experimental/GeometryCache/Source/GeometryCache/Private/NormalCompression.cpp
--- a/Engine/Plugins/Experimental/GeometryCache/Source/GeometryCache/Private/NormalCompression.cpp
+++ b/Engine/Plugins/Experimental/GeometryCache/Source/GeometryCache/Private/NormalCompression.cpp
@@ -115,15 +115,15 @@ FNormalCoderSmith::FSphericalCoordinates FNormalCoderSmith::VectorToSpherical(co
 {
 	FVector NormalizedVector = Vector;
 	NormalizedVector.Normalize();
-	FVector2D UE4Spherical = NormalizedVector.UnitCartesianToSpherical();
-	return FNormalCoderSmith::FSphericalCoordinates(UE4Spherical.X, UE4Spherical.Y);
+	FVector2D UESpherical = NormalizedVector.UnitCartesianToSpherical();
+	return FNormalCoderSmith::FSphericalCoordinates(UESpherical.X, UESpherical.Y);
 }
 
 /** Convert from spherical coordinates to Cartesian vectors */
 FVector FNormalCoderSmith::SphericalToVector(const FNormalCoderSmith::FSphericalCoordinates& Spherical)
 {
-	FVector2D UE4Spherical(Spherical.Phi, Spherical.Theta);
-	return UE4Spherical.SphericalToUnitCartesian();
+	FVector2D UESpherical(Spherical.Phi, Spherical.Theta);
+	return UESpherical.SphericalToUnitCartesian();
 }
 
 FVector FNormalCoderSmith::MoveFrameEncode(const FVector& Value)
diff --git a/Engine/Plugins/Experimental/GeometryCache/Source/GeometryCache/Private/NormalCompression.h b/Engine/Plugins/Experimental/GeometryCache/Source/GeometryCache/Private/NormalCompression.h
--- a/Engine/Plugins/Experimental/GeometryCache/Source/GeometryCache/Private/NormalCompression.h
+++ b/Engine/Plugins/Experimental/GeometryCache/Source/GeometryCache/Private/NormalCompression.h
@@ -10,7 +10,7 @@ DECLARE_LOG_CATEGORY_EXTERN(LogGeoCaStreamingNormalCompression, Verbose, All);
 	spherical coordinates. Computers & Graphics 36, 5, 360-365.].
 	States they have better complexity/rate/quality properties than Octa (e.g., Griffith et al.) and Sextant (Deering et al.).
 	
-	Initial experiment and port to UE4, lots that can be optimized.
+	Initial experiment and port to UnrealEngine, lots that can be optimized.
 */
 class FNormalCoderSmith
 {
@@ -94,7 +94,7 @@ private:
 		return (i % n + n) % n;
 	}
 
-	/** Helpers to get our column vectors in our matrix. These can probably be replaced by UE4-specific operations. */
+	/** Helpers to get our column vectors in our matrix. These can probably be replaced by UnrealEngine-specific operations. */
 	FORCEINLINE FVector GetMatrixColumn(FMatrix& Matrix, int32 Column);
 	FORCEINLINE void SetMatrixColumn(FMatrix& Matrix, int32 Column, const FVector& Value);	
 };
