Fixup GeometryCacheStreamingManager to enable self removal from StreamingManager.

#rnx

diff --git a/Engine/Plugins/Experimental/GeometryCache/Source/GeometryCache/Private/GeometryCacheModule.cpp b/Engine/Plugins/Experimental/GeometryCache/Source/GeometryCache/Private/GeometryCacheModule.cpp
--- a/Engine/Plugins/Experimental/GeometryCache/Source/GeometryCache/Private/GeometryCacheModule.cpp
+++ b/Engine/Plugins/Experimental/GeometryCache/Source/GeometryCache/Private/GeometryCacheModule.cpp
@@ -4,6 +4,7 @@
 #if WITH_EDITOR
 #include "GeometryCacheEdModule.h"
 #endif // WITH_EDITOR
+#include "GeometryCacheStreamingManager.h"
 #include "CodecV1.h"
 
 IMPLEMENT_MODULE(FGeometryCacheModule, GeometryCache)
@@ -14,9 +15,12 @@ void FGeometryCacheModule::StartupModule()
 	FGeometryCacheEdModule& Module = FModuleManager::LoadModuleChecked<FGeometryCacheEdModule>(TEXT("GeometryCacheEd"));
 #endif
 
+	IGeometryCacheStreamingManager::Register();
+
 	FCodecV1Decoder::InitLUT();
 }
 
 void FGeometryCacheModule::ShutdownModule()
 {
+	IGeometryCacheStreamingManager::Unregister();
 }
diff --git a/Engine/Plugins/Experimental/GeometryCache/Source/GeometryCache/Private/GeometryCacheStreamingManager.cpp b/Engine/Plugins/Experimental/GeometryCache/Source/GeometryCache/Private/GeometryCacheStreamingManager.cpp
--- a/Engine/Plugins/Experimental/GeometryCache/Source/GeometryCache/Private/GeometryCacheStreamingManager.cpp
+++ b/Engine/Plugins/Experimental/GeometryCache/Source/GeometryCache/Private/GeometryCacheStreamingManager.cpp
@@ -49,20 +49,19 @@ struct FGeometryCacheStreamingManager : public IGeometryCacheStreamingManager
 	virtual void AddLevel(class ULevel* Level) override;
 	virtual void RemoveLevel(class ULevel* Level) override;
 	virtual void NotifyLevelOffset(class ULevel* Level, const FVector& Offset) override;
-
 	// End IStreamingManager interface
 
 	// IGeometryCacheStreamingManager interface
-	virtual void AddGeometryCache(UGeometryCacheTrackStreamable* Cache);
-	virtual void RemoveGeometryCache(UGeometryCacheTrackStreamable* Cache);
-	virtual bool IsManagedGeometryCache(const UGeometryCacheTrackStreamable* Cache) const;
-	virtual bool IsStreamingInProgress(const UGeometryCacheTrackStreamable* Cache);
-	virtual void AddStreamingComponent(UGeometryCacheComponent* CacheComponent);
-	virtual void RemoveStreamingComponent(UGeometryCacheComponent* CacheComponent);
-	virtual bool IsManagedComponent(const UGeometryCacheComponent* CacheComponent) const;
-	virtual void PrefetchData(UGeometryCacheComponent* CacheComponent);
-	virtual const uint8* MapChunk(const UGeometryCacheTrackStreamable* Track, uint32 ChunkIndex, uint32* OutChunkSize = nullptr);
-	virtual void UnmapChunk(const UGeometryCacheTrackStreamable* Track, uint32 ChunkIndex);
+	virtual void AddGeometryCache(UGeometryCacheTrackStreamable* Cache) override;
+	virtual void RemoveGeometryCache(UGeometryCacheTrackStreamable* Cache) override;
+	virtual bool IsManagedGeometryCache(const UGeometryCacheTrackStreamable* Cache) const override;
+	virtual bool IsStreamingInProgress(const UGeometryCacheTrackStreamable* Cache) override;
+	virtual void AddStreamingComponent(UGeometryCacheComponent* CacheComponent) override;
+	virtual void RemoveStreamingComponent(UGeometryCacheComponent* CacheComponent) override;
+	virtual bool IsManagedComponent(const UGeometryCacheComponent* CacheComponent) const override;
+	virtual void PrefetchData(UGeometryCacheComponent* CacheComponent) override;
+	virtual const uint8* MapChunk(const UGeometryCacheTrackStreamable* Track, uint32 ChunkIndex, uint32* OutChunkSize = nullptr) override;
+	virtual void UnmapChunk(const UGeometryCacheTrackStreamable* Track, uint32 ChunkIndex) override;
 	// End IGeometryCacheStreamingManager interface
 
 private:
@@ -77,29 +76,15 @@ private:
 
 	mutable FCriticalSection CriticalSection;
 
-	double LastTickTime;	
+	double LastTickTime;
 };
 
-
-static FGeometryCacheStreamingManager* GeometryStreamingManager = nullptr;
-
-IGeometryCacheStreamingManager &IGeometryCacheStreamingManager::Get()
-{
-	if (GeometryStreamingManager == nullptr)
-	{
-		GeometryStreamingManager = new FGeometryCacheStreamingManager();
-	}
-	return *GeometryStreamingManager;
-}
-
 FGeometryCacheStreamingManager::FGeometryCacheStreamingManager()
 {
-	IStreamingManager::Get().AddStreamingManager(this);
 }
 
 FGeometryCacheStreamingManager::~FGeometryCacheStreamingManager()
 {
-	IStreamingManager::Get().RemoveStreamingManager(this);
 }
 
 void FGeometryCacheStreamingManager::UpdateResourceStreaming(float DeltaTime, bool bProcessEverything)
@@ -346,3 +331,19 @@ void FGeometryCacheStreamingManager::UnmapChunk(const UGeometryCacheTrackStreama
 		(*data)->UnmapChunk(ChunkIndex);
 	}
 }
+
+IGeometryCacheStreamingManager &IGeometryCacheStreamingManager::Get()
+{
+	static FGeometryCacheStreamingManager Instance;
+	return Instance;
+}
+
+void IGeometryCacheStreamingManager::Register()
+{
+	IStreamingManager::Get().AddStreamingManager(&IGeometryCacheStreamingManager::Get());
+}
+
+void IGeometryCacheStreamingManager::Unregister()
+{
+	IStreamingManager::Get().RemoveStreamingManager(&IGeometryCacheStreamingManager::Get());
+}
diff --git a/Engine/Plugins/Experimental/GeometryCache/Source/GeometryCache/Private/GeometryCacheStreamingManager.h b/Engine/Plugins/Experimental/GeometryCache/Source/GeometryCache/Private/GeometryCacheStreamingManager.h
--- a/Engine/Plugins/Experimental/GeometryCache/Source/GeometryCache/Private/GeometryCacheStreamingManager.h
+++ b/Engine/Plugins/Experimental/GeometryCache/Source/GeometryCache/Private/GeometryCacheStreamingManager.h
@@ -32,6 +32,12 @@ struct IGeometryCacheStreamingManager : public IStreamingManager
 	/** Getter of the singleton */
 	GEOMETRYCACHE_API static struct IGeometryCacheStreamingManager& Get();
 
+	/** Streaming manager interface registration enrollment */
+	GEOMETRYCACHE_API static void Register();
+
+	/** Streaming manager interface registration removal */
+	GEOMETRYCACHE_API static void Unregister();
+
 	/** Adds a new cache to the streaming manager. */
 	virtual void AddGeometryCache(UGeometryCacheTrackStreamable* Cache) = 0;
 
