Fixed some geometry types were not getting into the ray tracing world

diff --git a/Engine/Plugins/Experimental/GeometryCache/Source/GeometryCache/Private/GeometryCacheSceneProxy.cpp b/Engine/Plugins/Experimental/GeometryCache/Source/GeometryCache/Private/GeometryCacheSceneProxy.cpp
--- a/Engine/Plugins/Experimental/GeometryCache/Source/GeometryCache/Private/GeometryCacheSceneProxy.cpp
+++ b/Engine/Plugins/Experimental/GeometryCache/Source/GeometryCache/Private/GeometryCacheSceneProxy.cpp
@@ -506,11 +506,6 @@ void FGeometryCacheSceneProxy::GetDynamicRayTracingInstances(FRayTracingMaterial
 			continue;
 		}
 
-		if (!TrackProxy->VertexFactory.GetType()->SupportsRayTracingDynamicGeometry())
-		{
-			continue;
-		}
-
 		FRayTracingInstance RayTracingInstance;
 		RayTracingInstance.Geometry = &TrackProxy->RayTracingGeometry;
 		RayTracingInstance.InstanceTransforms.Add(GetLocalToWorld());
diff --git a/Engine/Plugins/FX/Niagara/Source/Niagara/Private/NiagaraRendererMeshes.cpp b/Engine/Plugins/FX/Niagara/Source/Niagara/Private/NiagaraRendererMeshes.cpp
--- a/Engine/Plugins/FX/Niagara/Source/Niagara/Private/NiagaraRendererMeshes.cpp
+++ b/Engine/Plugins/FX/Niagara/Source/Niagara/Private/NiagaraRendererMeshes.cpp
@@ -1133,11 +1133,6 @@ void FNiagaraRendererMeshes::GetDynamicRayTracingInstances(FRayTracingMaterialGa
 		VertexFactory.InitResource();
 		SetupVertexFactory(VertexFactory, LODModel);
 
-		if (!VertexFactory.GetType()->SupportsRayTracingDynamicGeometry())
-		{
-			continue;
-		}
-
 		PreparePerMeshData(ParticleMeshRenderData, *SceneProxy, MeshData);
 
 		// Sort/Cull particles if needed.
