[Alembic Importer / Geometry Cache]:

Added support for optionally storing the imported vertex numbers from the DCC.
A cube might turn into 24 vertices after import, but now you can retrieve the vertex number in the DCC for each of those 24 vertices. For a cube these will range from 0 to 7 as the cube has 8 vertices in the DCC.
This is disabled on default and the assets remain backward compatible.

The new checkbox option under the geometry cache import options is called: Store Imported Vertex Numbers.

diff --git a/Engine/Plugins/Experimental/AlembicImporter/Source/AlembicLibrary/Private/AbcImportUtilities.cpp b/Engine/Plugins/Experimental/AlembicImporter/Source/AlembicLibrary/Private/AbcImportUtilities.cpp
--- a/Engine/Plugins/Experimental/AlembicImporter/Source/AlembicLibrary/Private/AbcImportUtilities.cpp
+++ b/Engine/Plugins/Experimental/AlembicImporter/Source/AlembicLibrary/Private/AbcImportUtilities.cpp
@@ -1507,7 +1507,7 @@ void AbcImporterUtilities::ApplyConversion(TArray<FMatrix>& InOutMatrices, const
 }
 
 void AbcImporterUtilities::GeometryCacheDataForMeshSample(FGeometryCacheMeshData &OutMeshData, const FAbcMeshSample* MeshSample,
-	const uint32 MaterialOffset, const float SecondsPerFrame, const bool bUseVelocitiesAsMotionVectors)
+	const uint32 MaterialOffset, const float SecondsPerFrame, const bool bUseVelocitiesAsMotionVectors, const bool bStoreImportedVertexNumbers)
 {
 	OutMeshData.BoundingBox = FBox(MeshSample->Vertices);
 
@@ -1516,12 +1516,14 @@ void AbcImporterUtilities::GeometryCacheDataForMeshSample(FGeometryCacheMeshData
 	const int32 NumNormals = MeshSample->Normals.Num();
 	const bool bHasTangents = MeshSample->TangentX.Num() == NumNormals && MeshSample->TangentY.Num() == NumNormals;
 	const bool bHasVelocities = bUseVelocitiesAsMotionVectors && (MeshSample->Velocities.Num() > 0);
+	const bool bHasImportedVertexNumbers = (NumNormals > 0) && bStoreImportedVertexNumbers;
 
 	OutMeshData.VertexInfo.bHasColor0 = true;
 	OutMeshData.VertexInfo.bHasTangentX = bHasTangents;
 	OutMeshData.VertexInfo.bHasTangentZ = NumNormals > 0;
 	OutMeshData.VertexInfo.bHasUV0 = true;
 	OutMeshData.VertexInfo.bHasMotionVectors = bHasVelocities;
+	OutMeshData.VertexInfo.bHasImportedVertexNumbers = bHasImportedVertexNumbers;
 
 	uint32 NumMaterials = MaterialOffset;
 
@@ -1547,6 +1549,11 @@ void AbcImporterUtilities::GeometryCacheDataForMeshSample(FGeometryCacheMeshData
 	OutMeshData.TextureCoordinates.AddZeroed(NumNormals);
 	OutMeshData.Colors.AddZeroed(NumNormals);
 
+	if (bHasImportedVertexNumbers)
+	{
+		OutMeshData.ImportedVertexNumbers.AddZeroed(NumNormals);
+	}
+
 	for (int32 TriangleIndex = 0; TriangleIndex < NumTriangles; ++TriangleIndex)
 	{
 		const int32 SectionIndex = MeshSample->MaterialIndices[TriangleIndex];
@@ -1559,6 +1566,11 @@ void AbcImporterUtilities::GeometryCacheDataForMeshSample(FGeometryCacheMeshData
 
 			OutMeshData.Positions[CornerIndex] = MeshSample->Vertices[Index];
 
+			if (bHasImportedVertexNumbers)
+			{
+				OutMeshData.ImportedVertexNumbers[CornerIndex] = Index;
+			}
+
 			if (bHasVelocities)
 			{
 				FVector MotionVector = MeshSample->Velocities[Index];
@@ -1610,7 +1622,7 @@ void AbcImporterUtilities::GeometryCacheDataForMeshSample(FGeometryCacheMeshData
 
 void AbcImporterUtilities::MergePolyMeshesToMeshData(int32 FrameIndex, int32 FrameStart, float SecondsPerFrame, bool bUseVelocitiesAsMotionVectors,
 	const TArray<FAbcPolyMesh*>& PolyMeshes, const TArray<FString>& UniqueFaceSetNames,
-	FGeometryCacheMeshData& MeshData, int32& PreviousNumVertices, bool& bConstantTopology)
+	FGeometryCacheMeshData& MeshData, int32& PreviousNumVertices, bool& bConstantTopology, bool bStoreImportedVertexNumbers)
 {
 	FAbcMeshSample MergedSample;
 
@@ -1658,7 +1670,7 @@ void AbcImporterUtilities::MergePolyMeshesToMeshData(int32 FrameIndex, int32 Fra
 	MergedSample.NumMaterials = UniqueFaceSetNames.Num();
 
 	// Generate the mesh data for this sample
-	AbcImporterUtilities::GeometryCacheDataForMeshSample(MeshData, &MergedSample, 0, SecondsPerFrame, bUseVelocitiesAsMotionVectors);
+	AbcImporterUtilities::GeometryCacheDataForMeshSample(MeshData, &MergedSample, 0, SecondsPerFrame, bUseVelocitiesAsMotionVectors, bStoreImportedVertexNumbers);
 }
 
 UMaterialInterface* AbcImporterUtilities::RetrieveMaterial(FAbcFile& AbcFile, const FString& MaterialName, UObject* InParent, EObjectFlags Flags)
diff --git a/Engine/Plugins/Experimental/AlembicImporter/Source/AlembicLibrary/Private/AbcImportUtilities.h b/Engine/Plugins/Experimental/AlembicImporter/Source/AlembicLibrary/Private/AbcImportUtilities.h
--- a/Engine/Plugins/Experimental/AlembicImporter/Source/AlembicLibrary/Private/AbcImportUtilities.h
+++ b/Engine/Plugins/Experimental/AlembicImporter/Source/AlembicLibrary/Private/AbcImportUtilities.h
@@ -330,7 +330,7 @@ namespace AbcImporterUtilities
 	bool IsObjectVisibilityConstant(const Alembic::Abc::IObject& Object);
 
 	/** Generates and populates a FGeometryCacheMeshData instance from and for the given mesh sample */
-	void GeometryCacheDataForMeshSample(FGeometryCacheMeshData &OutMeshData, const FAbcMeshSample* MeshSample, const uint32 MaterialOffset, const float SecondsPerFrame, const bool bUseVelocitiesAsMotionVectors);
+	void GeometryCacheDataForMeshSample(FGeometryCacheMeshData &OutMeshData, const FAbcMeshSample* MeshSample, const uint32 MaterialOffset, const float SecondsPerFrame, const bool bUseVelocitiesAsMotionVectors, const bool bStoreImportedVertexNumbers);
 
 	/**
 	 * Merges the given PolyMeshes at the given FrameIndex into a GeometryCacheMeshData
@@ -344,10 +344,11 @@ namespace AbcImporterUtilities
 	 * @param MeshData				The GeometryCacheMeshData where to output the merged PolyMeshes
 	 * @param PreviousNumVertices	The number of vertices in the merged PolyMeshes, used to determine if its topology is constant between 2 frames
 	 * @param bConstantTopology		Flag to indicate if the merged PolyMeshes has constant topology
+	 * @param bStoreImportedVertexNumbers Set to true when we want to store the original dcc vertex numbers for each vertex.
 	 */
 	void MergePolyMeshesToMeshData(int32 FrameIndex, int32 FrameStart, float SecondsPerFrame, bool bUseVelocitiesAsMotionVectors,
 		const TArray<FAbcPolyMesh*>& PolyMeshes, const TArray<FString>& UniqueFaceSetNames,
-		FGeometryCacheMeshData& MeshData, int32& PreviousNumVertices, bool& bConstantTopology);
+		FGeometryCacheMeshData& MeshData, int32& PreviousNumVertices, bool& bConstantTopology, bool bStoreImportedVertexNumbers);
 
 	/** Retrieves a material from an AbcFile according to the given name and resaves it into the parent package */
 	UMaterialInterface* RetrieveMaterial(FAbcFile& AbcFile, const FString& MaterialName, UObject* InParent, EObjectFlags Flags);
diff --git a/Engine/Plugins/Experimental/AlembicImporter/Source/AlembicLibrary/Private/AbcImporter.cpp b/Engine/Plugins/Experimental/AlembicImporter/Source/AlembicLibrary/Private/AbcImporter.cpp
--- a/Engine/Plugins/Experimental/AlembicImporter/Source/AlembicLibrary/Private/AbcImporter.cpp
+++ b/Engine/Plugins/Experimental/AlembicImporter/Source/AlembicLibrary/Private/AbcImporter.cpp
@@ -401,10 +401,11 @@ UGeometryCache* FAbcImporter::ImportAsGeometryCache(UObject* InParent, EObjectFl
 					const bool bUseVelocitiesAsMotionVectors = (ImportSettings->GeometryCacheSettings.MotionVectors == EAbcGeometryCacheMotionVectorsImport::ImportAbcVelocitiesAsMotionVectors);
 					FGeometryCacheMeshData MeshData;
 					bool bConstantTopology = true;
+					const bool bStoreImportedVertexNumbers = ImportSettings->GeometryCacheSettings.bStoreImportedVertexNumbers;
 
 					AbcImporterUtilities::MergePolyMeshesToMeshData(FrameIndex, ImportSettings->SamplingSettings.FrameStart, AbcFile->GetSecondsPerFrame(), bUseVelocitiesAsMotionVectors,
 						PolyMeshes, UniqueFaceSetNames,
-						MeshData, PreviousNumVertices, bConstantTopology);
+						MeshData, PreviousNumVertices, bConstantTopology, bStoreImportedVertexNumbers);
 					
 					Tracks[0]->AddMeshSample(MeshData, PolyMeshes[0]->GetTimeForFrameIndex(FrameIndex) - InAbcFile->GetImportTimeOffset(), bConstantTopology);
 					
@@ -488,9 +489,10 @@ UGeometryCache* FAbcImporter::ImportAsGeometryCache(UObject* InParent, EObjectFl
 							if (bVisible)
 							{
 								const bool bUseVelocitiesAsMotionVectors = ( ImportSettings->GeometryCacheSettings.MotionVectors == EAbcGeometryCacheMotionVectorsImport::ImportAbcVelocitiesAsMotionVectors );
+								const bool bStoreImportedVertexData = ImportSettings->GeometryCacheSettings.bStoreImportedVertexNumbers;
 								const FAbcMeshSample* Sample = PolyMesh->GetSample(FrameIndex);
 								FGeometryCacheMeshData MeshData;
-								AbcImporterUtilities::GeometryCacheDataForMeshSample(MeshData, Sample, MaterialOffsets[TrackIndex], AbcFile->GetSecondsPerFrame(), bUseVelocitiesAsMotionVectors);
+								AbcImporterUtilities::GeometryCacheDataForMeshSample(MeshData, Sample, MaterialOffsets[TrackIndex], AbcFile->GetSecondsPerFrame(), bUseVelocitiesAsMotionVectors, bStoreImportedVertexData);
 								Track->AddMeshSample(MeshData, FrameTime, PolyMesh->bConstantTopology);
 							}
 
diff --git a/Engine/Plugins/Experimental/AlembicImporter/Source/AlembicLibrary/Private/AbcUtilities.cpp b/Engine/Plugins/Experimental/AlembicImporter/Source/AlembicLibrary/Private/AbcUtilities.cpp
--- a/Engine/Plugins/Experimental/AlembicImporter/Source/AlembicLibrary/Private/AbcUtilities.cpp
+++ b/Engine/Plugins/Experimental/AlembicImporter/Source/AlembicLibrary/Private/AbcUtilities.cpp
@@ -15,9 +15,10 @@ void FAbcUtilities::GetFrameMeshData(FAbcFile& AbcFile, int32 FrameIndex, FGeome
 	int32 PreviousNumVertices = 0;
 	bool bConstantTopology = false;
 	const bool bUseVelocitiesAsMotionVectors = true;
+	const bool bStoreImportedVertexNumbers = false;
 
 	AbcImporterUtilities::MergePolyMeshesToMeshData(FrameIndex, 0, AbcFile.GetSecondsPerFrame(), bUseVelocitiesAsMotionVectors,
-		AbcFile.GetPolyMeshes(), AbcFile.GetUniqueFaceSetNames(), MeshData, PreviousNumVertices, bConstantTopology);
+		AbcFile.GetPolyMeshes(), AbcFile.GetUniqueFaceSetNames(), MeshData, PreviousNumVertices, bConstantTopology, bStoreImportedVertexNumbers);
 
 	OutMeshData = MoveTemp(MeshData);
 
diff --git a/Engine/Plugins/Experimental/AlembicImporter/Source/AlembicLibrary/Public/AbcImportSettings.h b/Engine/Plugins/Experimental/AlembicImporter/Source/AlembicLibrary/Public/AbcImportSettings.h
--- a/Engine/Plugins/Experimental/AlembicImporter/Source/AlembicLibrary/Public/AbcImportSettings.h
+++ b/Engine/Plugins/Experimental/AlembicImporter/Source/AlembicLibrary/Public/AbcImportSettings.h
@@ -262,6 +262,7 @@ struct FAbcGeometryCacheSettings
 
 	FAbcGeometryCacheSettings()
 	:	bFlattenTracks(true),
+		bStoreImportedVertexNumbers(false),
 		bApplyConstantTopologyOptimizations(false),
 		bCalculateMotionVectorsDuringImport_DEPRECATED(false),
 		MotionVectors(EAbcGeometryCacheMotionVectorsImport::NoMotionVectors),
@@ -274,6 +275,11 @@ struct FAbcGeometryCacheSettings
 	UPROPERTY(EditAnywhere, BlueprintReadWrite, Category = GeometryCache)
 	bool bFlattenTracks;
 
+	/** Store the imported vertex numbers. This lets you know the vertex numbers inside the DCC.
+	  * The values of each vertex number will range from 0 to 7 for a cube. Even if the number of positions might be 24. */
+	UPROPERTY(EditAnywhere, BlueprintReadWrite, Category = GeometryCache, AdvancedDisplay)
+	bool bStoreImportedVertexNumbers;
+
 	/** Force the preprocessor to only do optimization once instead of when the preprocessor decides. This may lead to some problems with certain meshes but makes sure motion
 	    blur always works if the topology is constant. */
 	UPROPERTY(EditAnywhere, BlueprintReadWrite, Category = GeometryCache)
diff --git a/Engine/Plugins/Experimental/GeometryCache/Source/GeometryCache/Classes/GeometryCacheMeshData.h b/Engine/Plugins/Experimental/GeometryCache/Source/GeometryCache/Classes/GeometryCacheMeshData.h
--- a/Engine/Plugins/Experimental/GeometryCache/Source/GeometryCache/Classes/GeometryCacheMeshData.h
+++ b/Engine/Plugins/Experimental/GeometryCache/Source/GeometryCache/Classes/GeometryCacheMeshData.h
@@ -44,6 +44,7 @@ struct FGeometryCacheVertexInfo
 	bool bHasUV0;
 	bool bHasColor0;
 	bool bHasMotionVectors;
+	bool bHasImportedVertexNumbers;
 
 	bool bConstantUV0;
 	bool bConstantColor0;
@@ -56,6 +57,7 @@ struct FGeometryCacheVertexInfo
 		bHasUV0 = false;
 		bHasColor0 = false;
 		bHasMotionVectors = false;
+		bHasImportedVertexNumbers = false;
 		bConstantUV0 = false;
 		bConstantColor0 = false;
 		bConstantIndices = false;
@@ -81,6 +83,7 @@ struct FGeometryCacheVertexInfo
 		Ar << Mesh.bHasUV0;
 		Ar << Mesh.bHasColor0;
 		Ar << Mesh.bHasMotionVectors;
+		Ar << Mesh.bHasImportedVertexNumbers;
 
 		Ar << Mesh.bConstantUV0;
 		Ar << Mesh.bConstantColor0;
@@ -104,6 +107,7 @@ struct GEOMETRYCACHE_API FGeometryCacheMeshData
 		TangentsX.Empty();
 		TangentsZ.Empty();
 		Colors.Empty();
+		ImportedVertexNumbers.Empty();
 
 		MotionVectors.Empty();
 		BatchesInfo.Empty();
@@ -117,6 +121,9 @@ struct GEOMETRYCACHE_API FGeometryCacheMeshData
 	TArray<FPackedNormal> TangentsZ;
 	TArray<FColor> Colors;
 
+	/** The original imported vertex numbers, as they come from the DCC. One for each vertex position. */
+	TArray<uint32> ImportedVertexNumbers;
+
 	/** Motion vector for each vertex. The number of motion vectors should be zero (= no motion vectors) or identical to the number of vertices. */
 	TArray<FVector3f> MotionVectors;
 	/** Array of per-batch info structs*/
@@ -124,7 +131,7 @@ struct GEOMETRYCACHE_API FGeometryCacheMeshData
 	/** Bounding box for this sample in the track */
 	FBox BoundingBox;
 	/** Indices for this sample, used for drawing the mesh */
-	TArray<uint32> Indices;
+	TArray<uint32> Indices;	
 	/** Info on the vertex attributes */
 	FGeometryCacheVertexInfo VertexInfo;
 		
@@ -133,7 +140,7 @@ struct GEOMETRYCACHE_API FGeometryCacheMeshData
 
 	/** Serialization for const FVertexAnimationSample. */
 	friend FArchive& operator<<(FArchive& Ar, const FGeometryCacheMeshData& Mesh)
-		{
+	{
 		check(Ar.IsSaving());
 		return (Ar << const_cast<FGeometryCacheMeshData&>(Mesh));
 	}
@@ -147,7 +154,6 @@ struct GEOMETRYCACHE_API FGeometryCacheMeshData
 		CumulativeResourceSize.AddUnknownMemoryBytes(TangentsZ.Num() * sizeof(FPackedNormal));
 		CumulativeResourceSize.AddUnknownMemoryBytes(Colors.Num() * sizeof(FColor));
 
-
 		CumulativeResourceSize.AddUnknownMemoryBytes(MotionVectors.Num() * sizeof(FVector3f));
 		CumulativeResourceSize.AddUnknownMemoryBytes(BatchesInfo.Num() * sizeof(FGeometryCacheMeshBatchInfo));
 
@@ -160,6 +166,7 @@ struct GEOMETRYCACHE_API FGeometryCacheMeshData
 		CumulativeResourceSize.AddUnknownMemoryBytes(sizeof(BatchesInfo));
 		CumulativeResourceSize.AddUnknownMemoryBytes(sizeof(BoundingBox));
 		CumulativeResourceSize.AddUnknownMemoryBytes(Indices.Num() * sizeof(uint32));
+		CumulativeResourceSize.AddUnknownMemoryBytes(ImportedVertexNumbers.Num() * sizeof(uint32));
 		CumulativeResourceSize.AddUnknownMemoryBytes(sizeof(Indices));
 		CumulativeResourceSize.AddUnknownMemoryBytes(sizeof(VertexInfo));
 	}
diff --git a/Engine/Plugins/Experimental/GeometryCache/Source/GeometryCache/Private/CodecV1.cpp b/Engine/Plugins/Experimental/GeometryCache/Source/GeometryCache/Private/CodecV1.cpp
--- a/Engine/Plugins/Experimental/GeometryCache/Source/GeometryCache/Private/CodecV1.cpp
+++ b/Engine/Plugins/Experimental/GeometryCache/Source/GeometryCache/Private/CodecV1.cpp
@@ -11,8 +11,9 @@
 #include "Stats/StatsMisc.h"
 #include "Async/ParallelFor.h"
 
-#define OLD_MAGIC 123
-#define CURRENT_MAGIC 124
+#define V1_MAGIC 123
+#define V2_MAGIC 124
+#define V3_MAGIC 125 // Introduction of ImportedVertexNumbers
 
 #define MAX_CHUNK_VERTICES	1024u
 #define CHUNK_MAX_TRIANGLES	2048u
@@ -507,6 +508,7 @@ void FCodecV1Encoder::WriteCodedStreamDescription()
 	WriteBits(VertexInfo.bHasUV0 ? 1 : 0, 1);
 	WriteBits(VertexInfo.bHasColor0 ? 1 : 0, 1);
 	WriteBits(VertexInfo.bHasMotionVectors ? 1 : 0, 1);
+	WriteBits(VertexInfo.bHasImportedVertexNumbers ? 1 : 0, 1);
 
 	WriteBits(VertexInfo.bConstantUV0 ? 1 : 0, 1);
 	WriteBits(VertexInfo.bConstantColor0 ? 1 : 0, 1);
@@ -584,6 +586,7 @@ bool FCodecV1Encoder::EncodeFrameData(FMemoryWriter& Writer, const FGeometryCach
 	const TArray<FPackedNormal>& TangentsX = MeshData.TangentsX;
 	const TArray<FPackedNormal>& TangentsZ = MeshData.TangentsZ;
 	const TArray<FColor>& Colors = MeshData.Colors;
+	const TArray<uint32>& ImportedVertexNumbers = MeshData.ImportedVertexNumbers;
 
 	const TArray<uint32>& Indices = MeshData.Indices;
 	const TArray<FVector3f>& MotionVectors = MeshData.MotionVectors;	
@@ -633,6 +636,12 @@ bool FCodecV1Encoder::EncodeFrameData(FMemoryWriter& Writer, const FGeometryCach
 		if (NumChunkVertices > 0)
 		{
 			EncodePositionStream(&Positions[VertexOffset], Positions.GetTypeSize(), NumChunkVertices, Statistics.Vertices);
+
+			if (VertexInfo.bHasImportedVertexNumbers)
+			{
+				EncodeIndexStream(&ImportedVertexNumbers[VertexOffset], ImportedVertexNumbers.GetTypeSize(), NumChunkVertices, Statistics.ImportedVertexNumbers);
+			}
+
 			if (VertexInfo.bHasColor0)
 			{
 				EncodeColorStream(&Colors[VertexOffset], Colors.GetTypeSize(), NumChunkVertices, Statistics.Colors);
@@ -662,7 +671,7 @@ bool FCodecV1Encoder::EncodeFrameData(FMemoryWriter& Writer, const FGeometryCach
 	{
 		// Write out bitstream
 		FCodedFrameHeader Header = { 0 };
-		Header.Magic = CURRENT_MAGIC;
+		Header.Magic = V3_MAGIC;
 		Header.VertexCount = (uint32)Positions.Num();
 		Header.IndexCount = (uint32)Indices.Num();
 		uint32 PayloadSize = BitWriter.GetNumBytes();
@@ -677,6 +686,7 @@ bool FCodecV1Encoder::EncodeFrameData(FMemoryWriter& Writer, const FGeometryCach
 	// Gather stats for all streams
 	const uint32 TotalRawSize =
 		sizeof(uint32) * Indices.Num()	// Indices
+		+ sizeof(uint32) * ImportedVertexNumbers.Num()	// Imported vertex numbers
 		+ sizeof(FVector3f) * Positions.Num() // Vertices
 		+ sizeof(FColor) * Colors.Num() // Colors
 		+ sizeof(FPackedNormal) * TangentsX.Num() // TangentX
@@ -753,6 +763,7 @@ void FCodecV1Encoder::SetupTables()
 	// Most tables store 32-bit integers stored with a bit-length;raw value scheme. Some store specific symbols.
 	EncodingContext.ResidualIndicesTable.Initialize(HuffmanTableInt32SymbolCount);
 	EncodingContext.ResidualVertexPosTable.Initialize(HuffmanTableInt32SymbolCount);
+	EncodingContext.ResidualImportedVertexNumbersTable.Initialize(HuffmanTableInt32SymbolCount);
 	EncodingContext.ResidualColorTable.Initialize(HuffmanTableInt32SymbolCount);
 	EncodingContext.ResidualNormalTangentXTable.Initialize(HuffmanTableInt8SymbolCount);
 	EncodingContext.ResidualNormalTangentZTable.Initialize(HuffmanTableInt8SymbolCount);
@@ -767,9 +778,17 @@ void FCodecV1Encoder::SetPrepass(bool bPrepass)
 
 	// When bPrepass is set to true, the tables gather statistics about the data they encounter and do not write 
 	// any output bits.When set to false, they build the internal symbol representations and will write bits.
-	if ( !VertexInfo.bConstantIndices)
-		EncodingContext.ResidualIndicesTable.SetPrepass(bPrepass);	
+	if (!VertexInfo.bConstantIndices)
+	{
+		EncodingContext.ResidualIndicesTable.SetPrepass(bPrepass);
+	}
+
 	EncodingContext.ResidualVertexPosTable.SetPrepass(bPrepass);
+
+	if (VertexInfo.bHasImportedVertexNumbers)
+	{
+		EncodingContext.ResidualImportedVertexNumbersTable.SetPrepass(bPrepass);
+	}
 	if (VertexInfo.bHasColor0)
 	{
 		EncodingContext.ResidualColorTable.SetPrepass(bPrepass);
@@ -802,9 +821,17 @@ void FCodecV1Encoder::WriteTables()
 	FHuffmanBitStreamWriter& Writer = *EncodingContext.Writer;	
 	const FGeometryCacheVertexInfo& VertexInfo = EncodingContext.MeshData->VertexInfo;
 
-	if ( !VertexInfo.bConstantIndices)
-		EncodingContext.ResidualIndicesTable.Serialize(Writer);	
+	if (!VertexInfo.bConstantIndices)
+	{
+		EncodingContext.ResidualIndicesTable.Serialize(Writer);
+	}
+
 	EncodingContext.ResidualVertexPosTable.Serialize(Writer);
+
+	if (VertexInfo.bHasImportedVertexNumbers)
+	{
+		EncodingContext.ResidualImportedVertexNumbersTable.Serialize(Writer);
+	}
 	if (VertexInfo.bHasColor0)
 	{
 		EncodingContext.ResidualColorTable.Serialize(Writer);
@@ -990,9 +1017,16 @@ void FCodecV1Decoder::SetupAndReadTables(FHuffmanBitStreamReader& Reader)
 	const FGeometryCacheVertexInfo& VertexInfo = DecodingContext.MeshData->VertexInfo;
 		
 	if (!VertexInfo.bConstantIndices)
-		DecodingContext.ResidualIndicesTable.Initialize(Reader);	
+	{
+		DecodingContext.ResidualIndicesTable.Initialize(Reader);
+	}
+
 	DecodingContext.ResidualVertexPosTable.Initialize(Reader);
-	
+
+	if (VertexInfo.bHasImportedVertexNumbers)
+	{
+		DecodingContext.ResidualImportedVertexNumbersTable.Initialize(Reader);
+	}
 	if (VertexInfo.bHasColor0)
 	{
 		DecodingContext.ResidualColorTable.Initialize(Reader);		
@@ -1017,7 +1051,7 @@ void FCodecV1Decoder::SetupAndReadTables(FHuffmanBitStreamReader& Reader)
 	// Add additional tables here
 }
 
-void FCodecV1Decoder::ReadCodedStreamDescription(FHuffmanBitStreamReader& Reader)
+void FCodecV1Decoder::ReadCodedStreamDescription(FHuffmanBitStreamReader& Reader, uint32 Version)
 {	
 	FGeometryCacheVertexInfo& VertexInfo = DecodingContext.MeshData->VertexInfo;
 	
@@ -1027,6 +1061,10 @@ void FCodecV1Decoder::ReadCodedStreamDescription(FHuffmanBitStreamReader& Reader
 	VertexInfo.bHasUV0 = (ReadBits(Reader, 1) == 1);
 	VertexInfo.bHasColor0 = (ReadBits(Reader, 1) == 1);
 	VertexInfo.bHasMotionVectors = (ReadBits(Reader, 1) == 1);
+	if (Version >= V3_MAGIC)
+	{
+		VertexInfo.bHasImportedVertexNumbers = (ReadBits(Reader, 1) == 1);
+	}
 
 	VertexInfo.bConstantUV0 = (ReadBits(Reader, 1) == 1);
 	VertexInfo.bConstantColor0 = (ReadBits(Reader, 1) == 1);
@@ -1049,6 +1087,7 @@ void FCodecV1Decoder::InitLUT()
 DECLARE_CYCLE_STAT(TEXT("FCodecV1Decoder"), STAT_CodecV1Decoder, STATGROUP_GeometryCache);
 DECLARE_CYCLE_STAT(TEXT("SetupAndReadTables"), STAT_SetupAndReadTables, STATGROUP_GeometryCache);
 DECLARE_CYCLE_STAT(TEXT("DecodeIndexStream"), STAT_DecodeIndexStream, STATGROUP_GeometryCache);
+DECLARE_CYCLE_STAT(TEXT("DecodeImportedVertexNumbersStream"), STAT_DecodeImportedVertexNumbersStream, STATGROUP_GeometryCache);
 DECLARE_CYCLE_STAT(TEXT("DecodePositionStream"), STAT_DecodePositionStream, STATGROUP_GeometryCache);
 DECLARE_CYCLE_STAT(TEXT("DecodeColorStream"), STAT_DecodeColorStream, STATGROUP_GeometryCache);
 DECLARE_CYCLE_STAT(TEXT("DecodeTangentXStream"), STAT_DecodeTangentXStream, STATGROUP_GeometryCache);
@@ -1065,7 +1104,7 @@ bool FCodecV1Decoder::DecodeFrameData(FBufferReader& Reader, FGeometryCacheMeshD
 	FCodedFrameHeader Header;
 	Reader.Serialize(&Header, sizeof(Header));
 
-	if (Header.Magic != OLD_MAGIC && Header.Magic != CURRENT_MAGIC)
+	if (Header.Magic != V1_MAGIC && Header.Magic != V2_MAGIC && Header.Magic != V3_MAGIC)
 	{
 		UE_LOG(LogGeoCaStreamingCodecV1, Error, TEXT("Incompatible bitstream found"));
 		return false;
@@ -1082,7 +1121,7 @@ bool FCodecV1Decoder::DecodeFrameData(FBufferReader& Reader, FGeometryCacheMeshD
 	TArray<uint32> ChunkOffsets;
 
 	bool IsChunked = false;
-	if (Header.Magic == OLD_MAGIC)
+	if (Header.Magic == V1_MAGIC)
 	{
 		// Compatibility with old files. Just treat everything as one big chunk
 		IsChunked = false;
@@ -1090,6 +1129,7 @@ bool FCodecV1Decoder::DecodeFrameData(FBufferReader& Reader, FGeometryCacheMeshD
 	}
 	else
 	{
+		check(Header.Magic == V2_MAGIC || Header.Magic == V3_MAGIC);
 		IsChunked = true;
 		Reader << ChunkOffsets;
 		NumChunks = ChunkOffsets.Num();
@@ -1105,7 +1145,7 @@ bool FCodecV1Decoder::DecodeFrameData(FBufferReader& Reader, FGeometryCacheMeshD
 	FHuffmanBitStreamReader BitReader(Bytes.GetData(), Bytes.Num());
 	
 	// Read which vertex attributes are in the bit stream
-	ReadCodedStreamDescription(BitReader);
+	ReadCodedStreamDescription(BitReader, Header.Magic);
 
 	// Restore entropy coding contexts
 	{
@@ -1124,6 +1164,12 @@ bool FCodecV1Decoder::DecodeFrameData(FBufferReader& Reader, FGeometryCacheMeshD
 		OutMeshData.Positions.Empty(Header.VertexCount);
 		OutMeshData.Positions.AddUninitialized(Header.VertexCount);
 
+		OutMeshData.ImportedVertexNumbers.Empty(Header.VertexCount);
+		if (OutMeshData.VertexInfo.bHasImportedVertexNumbers)
+		{
+			OutMeshData.ImportedVertexNumbers.AddUninitialized(Header.VertexCount);
+		}
+
 		OutMeshData.Colors.Empty(Header.VertexCount);
 		if (OutMeshData.VertexInfo.bHasColor0)
 		{
@@ -1132,8 +1178,7 @@ bool FCodecV1Decoder::DecodeFrameData(FBufferReader& Reader, FGeometryCacheMeshD
 		else
 		{
 			OutMeshData.Colors.AddZeroed(Header.VertexCount);
-		}
-		
+		}	
 		
 		OutMeshData.TangentsX.Empty(Header.VertexCount);
 		if (OutMeshData.VertexInfo.bHasTangentX)
@@ -1214,6 +1259,12 @@ bool FCodecV1Decoder::DecodeFrameData(FBufferReader& Reader, FGeometryCacheMeshD
 				DecodePositionStream(ChunkReader, &OutMeshData.Positions[VertexOffset], OutMeshData.Positions.GetTypeSize(), NumChunkVertices);
 			}
 
+			if (OutMeshData.VertexInfo.bHasImportedVertexNumbers)
+			{
+				SCOPE_CYCLE_COUNTER(STAT_DecodeImportedVertexNumbersStream);
+				DecodeIndexStream(ChunkReader, &OutMeshData.ImportedVertexNumbers[VertexOffset], OutMeshData.ImportedVertexNumbers.GetTypeSize(), NumChunkVertices);
+			}
+
 			if (OutMeshData.VertexInfo.bHasColor0)
 			{
 				SCOPE_CYCLE_COUNTER(STAT_DecodeColorStream);
diff --git a/Engine/Plugins/Experimental/GeometryCache/Source/GeometryCache/Private/CodecV1.h b/Engine/Plugins/Experimental/GeometryCache/Source/GeometryCache/Private/CodecV1.h
--- a/Engine/Plugins/Experimental/GeometryCache/Source/GeometryCache/Private/CodecV1.h
+++ b/Engine/Plugins/Experimental/GeometryCache/Source/GeometryCache/Private/CodecV1.h
@@ -87,6 +87,7 @@ struct FCodecV1DecodingContext
 	/** Huffman tables */
 	FHuffmanDecodeTable ResidualIndicesTable;	
 	FHuffmanDecodeTable ResidualVertexPosTable;
+	FHuffmanDecodeTable ResidualImportedVertexNumbersTable;
 	FHuffmanDecodeTable ResidualColorTable;
 	FHuffmanDecodeTable ResidualNormalTangentXTable;
 	FHuffmanDecodeTable ResidualNormalTangentZTable;
@@ -130,6 +131,7 @@ struct FCodecV1EncodingContext
 	/** Huffman tables */
 	FHuffmanEncodeTable ResidualIndicesTable;
 	FHuffmanEncodeTable ResidualVertexPosTable;
+	FHuffmanEncodeTable ResidualImportedVertexNumbersTable;
 	FHuffmanEncodeTable ResidualColorTable;
 	FHuffmanEncodeTable ResidualColorSkipTable;
 	FHuffmanEncodeTable ResidualNormalTangentXTable;
@@ -270,6 +272,7 @@ private:
 		
 		FStreamEncodingStatistics Indices;
 		FStreamEncodingStatistics Vertices;		
+		FStreamEncodingStatistics ImportedVertexNumbers;
 		FStreamEncodingStatistics Colors;
 		FStreamEncodingStatistics TangentX;
 		FStreamEncodingStatistics TangentY;
@@ -341,7 +344,7 @@ private:
 	}
 
 	/** Read info on the available streams from a bit stream */
-	void ReadCodedStreamDescription(FHuffmanBitStreamReader& Reader);
+	void ReadCodedStreamDescription(FHuffmanBitStreamReader& Reader, uint32 Version);
 
 	/** Any context information to decode a frame in a sequence of frames, such as any Huffman tables used. */
 	FCodecV1DecodingContext DecodingContext;	
diff --git a/Engine/Plugins/Experimental/GeometryCache/Source/GeometryCache/Private/GeometryCacheMeshData.cpp b/Engine/Plugins/Experimental/GeometryCache/Source/GeometryCache/Private/GeometryCacheMeshData.cpp
--- a/Engine/Plugins/Experimental/GeometryCache/Source/GeometryCache/Private/GeometryCacheMeshData.cpp
+++ b/Engine/Plugins/Experimental/GeometryCache/Source/GeometryCache/Private/GeometryCacheMeshData.cpp
@@ -40,6 +40,7 @@ FArchive& operator<<(FArchive& Ar, FGeometryCacheMeshData& Mesh)
 			Mesh.TangentsX.SetNumUninitialized(NumVertices);
 			Mesh.TangentsZ.SetNumUninitialized(NumVertices);
 			Mesh.Colors.SetNumUninitialized(NumVertices);
+			Mesh.ImportedVertexNumbers.SetNumUninitialized(NumVertices);
 
 			if (Mesh.VertexInfo.bHasMotionVectors)
 			{
@@ -58,11 +59,12 @@ FArchive& operator<<(FArchive& Ar, FGeometryCacheMeshData& Mesh)
 			Ar.Serialize(&Mesh.TangentsX[0], Mesh.TangentsX.Num() * Mesh.TangentsX.GetTypeSize());
 			Ar.Serialize(&Mesh.TangentsZ[0], Mesh.TangentsZ.Num() * Mesh.TangentsZ.GetTypeSize());
 			Ar.Serialize(&Mesh.Colors[0], Mesh.Colors.Num() * Mesh.Colors.GetTypeSize());/**/
+			Ar.Serialize(&Mesh.ImportedVertexNumbers, Mesh.ImportedVertexNumbers.Num() * Mesh.ImportedVertexNumbers.GetTypeSize());
 		  
 			if (Mesh.VertexInfo.bHasMotionVectors)
 			{
 				Ar.Serialize(&Mesh.MotionVectors[0], Mesh.MotionVectors.Num()*Mesh.MotionVectors.GetTypeSize());
-			}			
+			}
 		}
 	}
 
diff --git a/Engine/Plugins/Experimental/GeometryCache/Source/GeometryCache/Private/GeometryCachePreprocessor.cpp b/Engine/Plugins/Experimental/GeometryCache/Source/GeometryCache/Private/GeometryCachePreprocessor.cpp
--- a/Engine/Plugins/Experimental/GeometryCache/Source/GeometryCache/Private/GeometryCachePreprocessor.cpp
+++ b/Engine/Plugins/Experimental/GeometryCache/Source/GeometryCache/Private/GeometryCachePreprocessor.cpp
@@ -298,8 +298,15 @@ void FOptimizeGeometryCachePreprocessor::FlushBufferedFrames()
 
 		const bool bHasTangentsX = OldMesh.TangentsX.Num() > 0;
 		const bool bHasTangentsZ = OldMesh.TangentsZ.Num() > 0;
+		const bool bHasImportedVertexNumbers = OldMesh.ImportedVertexNumbers.Num() > 0;
 
 		NewMesh.Positions.SetNumUninitialized(NewVerticesReordered.Num());
+
+		if (bHasImportedVertexNumbers)
+		{
+			NewMesh.ImportedVertexNumbers.SetNumUninitialized(NewVerticesReordered.Num());
+		}
+
 		if (bHasTangentsX)
 		{
 			NewMesh.TangentsX.SetNumUninitialized(NewVerticesReordered.Num());
@@ -329,6 +336,10 @@ void FOptimizeGeometryCachePreprocessor::FlushBufferedFrames()
 		for (int32 i = 0; i < NewMesh.Positions.Num(); i++)
 		{
 			NewMesh.Positions[i] = OldMesh.Positions[NewVerticesReordered[i]];
+			if (bHasImportedVertexNumbers)
+			{
+				NewMesh.ImportedVertexNumbers[i] = OldMesh.ImportedVertexNumbers[NewVerticesReordered[i]];
+			}
 			if (bHasTangentsX)
 			{
 				NewMesh.TangentsX[i] = OldMesh.TangentsX[NewVerticesReordered[i]];
