Added a new renderer for geometry caches

diff --git a/Engine/Plugins/Experimental/GeometryCache/GeometryCache.uplugin b/Engine/Plugins/Experimental/GeometryCache/GeometryCache.uplugin
--- a/Engine/Plugins/Experimental/GeometryCache/GeometryCache.uplugin
+++ b/Engine/Plugins/Experimental/GeometryCache/GeometryCache.uplugin
@@ -14,6 +14,12 @@
 	"CanContainContent" : false,
 	"IsBetaVersion" : true,
 	"Installed" : false,
+	"Plugins": [
+		{
+			"Name": "Niagara",
+			"Enabled": true
+		}
+	],
 	"Modules" :
 	[
 		{
diff --git a/Engine/Plugins/Experimental/GeometryCache/Source/GeometryCache/Classes/NiagaraGeometryCacheRendererProperties.h b/Engine/Plugins/Experimental/GeometryCache/Source/GeometryCache/Classes/NiagaraGeometryCacheRendererProperties.h
new file mode 100644
--- /dev/null
+++ b/Engine/Plugins/Experimental/GeometryCache/Source/GeometryCache/Classes/NiagaraGeometryCacheRendererProperties.h
@@ -0,0 +1,129 @@
+// Copyright Epic Games, Inc. All Rights Reserved.
+
+#pragma once
+
+#include "CoreMinimal.h"
+#include "GeometryCache.h"
+#include "NiagaraCommon.h"
+#include "NiagaraRendererProperties.h"
+#include "NiagaraGeometryCacheRendererProperties.generated.h"
+
+class FNiagaraEmitterInstance;
+class SWidget;
+
+USTRUCT()
+struct FNiagaraGeometryCacheReference
+{
+	GENERATED_USTRUCT_BODY()
+public:
+	FNiagaraGeometryCacheReference();
+	
+	/** Reference to the geometry cache asset to use (if the user parameter binding is not set) */
+	UPROPERTY(EditAnywhere, Category = "GeometryCache")
+	TObjectPtr<UGeometryCache> GeometryCache;
+
+	/** Use the UGeometryCache bound to this user variable if it is set to a valid value. If this is bound to a valid value and GeometryCache is also set, GeometryCacheUserParamBinding wins.*/
+	UPROPERTY(EditAnywhere, Category = "GeometryCache")
+	FNiagaraUserParameterBinding GeometryCacheUserParamBinding;
+
+	/** The materials to be used instead of the GeometryCache's materials. If the GeometryCache requires more materials than exist in this array or any entry in this array is set to None, we will use the GeometryCache's existing material instead.*/
+	UPROPERTY(EditAnywhere, Category = "GeometryCache")
+	TArray<TObjectPtr<UMaterialInterface>> OverrideMaterials;
+};
+
+UCLASS(editinlinenew, MinimalAPI, meta = (DisplayName = "Geometry Cache Renderer"))
+class UNiagaraGeometryCacheRendererProperties : public UNiagaraRendererProperties
+{
+public:
+	GENERATED_BODY()
+
+	UNiagaraGeometryCacheRendererProperties();
+
+	//UObject Interface
+	virtual void PostLoad() override;
+	virtual void PostInitProperties() override;
+	//UObject Interface END
+
+	static void InitCDOPropertiesAfterModuleStartup();
+
+	//~ UNiagaraRendererProperties interface
+	virtual FNiagaraRenderer* CreateEmitterRenderer(ERHIFeatureLevel::Type FeatureLevel, const FNiagaraEmitterInstance* Emitter, const FNiagaraSystemInstanceController& InController) override;
+	virtual class FNiagaraBoundsCalculator* CreateBoundsCalculator() override { return nullptr; }
+	virtual bool IsSimTargetSupported(ENiagaraSimTarget InSimTarget) const override { return InSimTarget == ENiagaraSimTarget::CPUSim; };
+	virtual void GetUsedMaterials(const FNiagaraEmitterInstance* InEmitter, TArray<UMaterialInterface*>& OutMaterials) const override;
+	virtual bool PopulateRequiredBindings(FNiagaraParameterStore& InParameterStore)  override;
+
+#if WITH_EDITORONLY_DATA
+	virtual void GetRendererWidgets(const FNiagaraEmitterInstance* InEmitter, TArray<TSharedPtr<SWidget>>& OutWidgets, TSharedPtr<FAssetThumbnailPool> InThumbnailPool) const override;
+	virtual const FSlateBrush* GetStackIcon() const override;
+	virtual void GetRendererTooltipWidgets(const FNiagaraEmitterInstance* InEmitter, TArray<TSharedPtr<SWidget>>& OutWidgets, TSharedPtr<FAssetThumbnailPool> InThumbnailPool) const override;
+	virtual const TArray<FNiagaraVariable>& GetOptionalAttributes() override;
+	virtual bool IsMaterialValidForRenderer(UMaterial* Material, FText& InvalidMessage) override;
+	virtual void FixMaterial(UMaterial* Material) override;
+	virtual void GetRendererFeedback(const UNiagaraEmitter* InEmitter, TArray<FText>& OutErrors, TArray<FText>& OutWarnings, TArray<FText>& OutInfo) const override;
+#endif // WITH_EDITORONLY_DATA
+
+	/** Reference to the geometry cache assets to use. If ArrayIndexBinding is not set, a random element is used for each particle. */
+	UPROPERTY(EditAnywhere, Category = "GeometryCache")
+	TArray<FNiagaraGeometryCacheReference> GeometryCaches;
+
+	/** If true, then the geometry cache keeps playing in a loop */
+	UPROPERTY(EditAnywhere, Category = "GeometryCache")
+	bool bIsLooping = true;
+
+	/** The max number of components that this emitter will spawn or update each frame. */
+	UPROPERTY(EditAnywhere, Category = "GeometryCache", meta = (ClampMin = 1))
+	uint32 ComponentCountLimit = 15;
+	
+	/** Which attribute should we use for the geometry cache position? */
+	UPROPERTY(EditAnywhere, Category = "GeometryCache")
+	FNiagaraVariableAttributeBinding PositionBinding;
+
+	/** Which attribute should we use for the geometry cache rotation? */
+	UPROPERTY(EditAnywhere, Category = "GeometryCache")
+	FNiagaraVariableAttributeBinding RotationBinding;
+
+	/** Which attribute should we use for the geometry cache component scale? */
+	UPROPERTY(EditAnywhere, Category = "GeometryCache")
+	FNiagaraVariableAttributeBinding ScaleBinding;
+
+	/** Which attribute should we use for the geometry cache's current animation time? */
+	UPROPERTY(EditAnywhere, Category = "GeometryCache")
+	FNiagaraVariableAttributeBinding ElapsedTimeBinding;
+
+	/** Which attribute should we use to check if rendering should be enabled? */
+	UPROPERTY(EditAnywhere, Category = "GeometryCache")
+	FNiagaraVariableAttributeBinding EnabledBinding;
+
+	/** Which attribute should we use to pick the element in the geometry cache array for this renderer? If not set, a random element will be used. */
+	UPROPERTY(EditAnywhere, Category = "GeometryCache")
+	FNiagaraVariableAttributeBinding ArrayIndexBinding;
+
+	/** Which attribute should we use for the renderer visibility tag? */
+	UPROPERTY(EditAnywhere, Category = "GeometryCache")
+	FNiagaraVariableAttributeBinding RendererVisibilityTagBinding;
+
+	/** If a render visibility tag is present, particles whose tag matches this value will be visible in this renderer. */
+	UPROPERTY(EditAnywhere, Category = "GeometryCache")
+	int32 RendererVisibility;
+
+	/** If true then components will not be automatically assigned to the first particle available, but try to stick to the same particle based on its unique id.
+	* Disabling this option is faster, but a particle can get a different component each tick, which can lead to problems with for example motion blur. */
+	UPROPERTY(EditAnywhere, AdvancedDisplay, Category = "GeometryCache")
+	bool bAssignComponentsOnParticleID = true;
+
+	virtual bool NeedsSystemPostTick() const override { return true; }
+	virtual bool NeedsSystemCompletion() const override { return true; }
+
+protected:
+	void InitBindings();
+
+	static FNiagaraVariable Particles_GeoCacheRotation;
+	static FNiagaraVariable Particles_Age;
+	static FNiagaraVariable Particles_GeoCacheIsEnabled;
+	static FNiagaraVariable Particles_ArrayIndex;
+
+	static void InitDefaultAttributes();
+private:
+	static TArray<TWeakObjectPtr<UNiagaraGeometryCacheRendererProperties>> GeometryCacheRendererPropertiesToDeferredInit;
+};
diff --git a/Engine/Plugins/Experimental/GeometryCache/Source/GeometryCache/Classes/NiagaraRendererGeometryCache.h b/Engine/Plugins/Experimental/GeometryCache/Source/GeometryCache/Classes/NiagaraRendererGeometryCache.h
new file mode 100644
--- /dev/null
+++ b/Engine/Plugins/Experimental/GeometryCache/Source/GeometryCache/Classes/NiagaraRendererGeometryCache.h
@@ -0,0 +1,47 @@
+// Copyright Epic Games, Inc. All Rights Reserved.
+
+/*========================================================================================
+NiagaraRendererGeometryCache.h: Renderer for geometry cache components.
+=========================================================================================*/
+#pragma once
+
+#include "NiagaraRenderer.h"
+#include "CoreMinimal.h"
+#include "GeometryCacheComponent.h"
+
+class UNiagaraGeometryCacheRendererProperties;
+class FNiagaraDataSet;
+
+/**
+* NiagaraRendererGeometryCache renders a geometry cache asset
+*/
+class FNiagaraRendererGeometryCache : public FNiagaraRenderer
+{
+public:
+
+	FNiagaraRendererGeometryCache(ERHIFeatureLevel::Type FeatureLevel, const UNiagaraRendererProperties *InProps, const FNiagaraEmitterInstance* Emitter);
+	virtual ~FNiagaraRendererGeometryCache();
+
+	//FNiagaraRenderer interface
+	virtual void DestroyRenderState_Concurrent() override;
+	virtual void PostSystemTick_GameThread(const UNiagaraRendererProperties* InProperties, const FNiagaraEmitterInstance* Emitter) override;
+	virtual void OnSystemComplete_GameThread(const UNiagaraRendererProperties* InProperties, const FNiagaraEmitterInstance* Emitter) override;
+	//FNiagaraRenderer interface END
+
+private:
+	struct FComponentPoolEntry
+	{
+		TWeakObjectPtr<UGeometryCacheComponent> Component;
+		float LastActiveTime = 0.0f;
+		float LastElapsedTime = 0.0f;
+		int32 LastAssignedToParticleID = -1;
+	};
+
+	void ResetComponentPool(bool bResetOwner);
+
+	// if the niagara component is not attached to an actor, we need to spawn and keep track of a temporary actor
+	TWeakObjectPtr<AActor> SpawnedOwner;
+
+	// all of the spawned components
+	TArray<FComponentPoolEntry> ComponentPool;
+};
\ No newline at end of file
diff --git a/Engine/Plugins/Experimental/GeometryCache/Source/GeometryCache/GeometryCache.Build.cs b/Engine/Plugins/Experimental/GeometryCache/Source/GeometryCache/GeometryCache.Build.cs
--- a/Engine/Plugins/Experimental/GeometryCache/Source/GeometryCache/GeometryCache.Build.cs
+++ b/Engine/Plugins/Experimental/GeometryCache/Source/GeometryCache/GeometryCache.Build.cs
@@ -17,7 +17,9 @@ public class GeometryCache : ModuleRules
 				"SlateCore",
                 "InputCore",
                 "RenderCore",
-                "RHI"
+                "RHI",
+                "Niagara",
+                "NiagaraCore"
 			}
 		);
 
@@ -28,6 +30,7 @@ public class GeometryCache : ModuleRules
             PublicIncludePathModuleNames.Add("GeometryCacheEd");
             DynamicallyLoadedModuleNames.Add("GeometryCacheEd");
             PrivateDependencyModuleNames.Add("MeshUtilitiesCommon");
+            PrivateDependencyModuleNames.Add("UnrealEd");
         }        
 	}
 }
diff --git a/Engine/Plugins/Experimental/GeometryCache/Source/GeometryCache/Private/GeometryCacheModule.cpp b/Engine/Plugins/Experimental/GeometryCache/Source/GeometryCache/Private/GeometryCacheModule.cpp
--- a/Engine/Plugins/Experimental/GeometryCache/Source/GeometryCache/Private/GeometryCacheModule.cpp
+++ b/Engine/Plugins/Experimental/GeometryCache/Source/GeometryCache/Private/GeometryCacheModule.cpp
@@ -6,6 +6,7 @@
 #endif // WITH_EDITOR
 #include "GeometryCacheStreamingManager.h"
 #include "CodecV1.h"
+#include "NiagaraGeometryCacheRendererProperties.h"
 
 IMPLEMENT_MODULE(FGeometryCacheModule, GeometryCache)
 
@@ -18,6 +19,8 @@ void FGeometryCacheModule::StartupModule()
 	IGeometryCacheStreamingManager::Register();
 
 	FCodecV1Decoder::InitLUT();
+
+	UNiagaraGeometryCacheRendererProperties::InitCDOPropertiesAfterModuleStartup();
 }
 
 void FGeometryCacheModule::ShutdownModule()
diff --git a/Engine/Plugins/Experimental/GeometryCache/Source/GeometryCache/Private/GeometryCacheStreamingManager.h b/Engine/Plugins/Experimental/GeometryCache/Source/GeometryCache/Private/GeometryCacheStreamingManager.h
--- a/Engine/Plugins/Experimental/GeometryCache/Source/GeometryCache/Private/GeometryCacheStreamingManager.h
+++ b/Engine/Plugins/Experimental/GeometryCache/Source/GeometryCache/Private/GeometryCacheStreamingManager.h
@@ -60,7 +60,7 @@ struct IGeometryCacheStreamingManager : public IStreamingManager
 	 * this may be useful when the component has seeked etc.*/
 	virtual void PrefetchData(UGeometryCacheComponent* CacheComponent) = 0;
 
-	/** Returns true if this is a streaming Sound Source that is managed by the streaming manager. */
+	/** Returns true if this is a streaming component that is managed by the streaming manager. */
 	virtual bool IsManagedComponent(const UGeometryCacheComponent* CacheComponent) const = 0;
 
 	/**
diff --git a/Engine/Plugins/Experimental/GeometryCache/Source/GeometryCache/Private/NiagaraGeometryCacheRendererProperties.cpp b/Engine/Plugins/Experimental/GeometryCache/Source/GeometryCache/Private/NiagaraGeometryCacheRendererProperties.cpp
new file mode 100644
--- /dev/null
+++ b/Engine/Plugins/Experimental/GeometryCache/Source/GeometryCache/Private/NiagaraGeometryCacheRendererProperties.cpp
@@ -0,0 +1,360 @@
+// Copyright Epic Games, Inc. All Rights Reserved.
+
+#include "NiagaraGeometryCacheRendererProperties.h"
+#include "NiagaraRenderer.h"
+#include "NiagaraConstants.h"
+#include "NiagaraRendererGeometryCache.h"
+#include "NiagaraComponent.h"
+#include "Modules/ModuleManager.h"
+#if WITH_EDITOR
+#include "Widgets/Images/SImage.h"
+#include "Styling/SlateIconFinder.h"
+#include "Widgets/SWidget.h"
+#include "Styling/SlateBrush.h"
+#include "AssetThumbnail.h"
+#endif
+
+#define LOCTEXT_NAMESPACE "UNiagaraGeometryCacheRendererProperties"
+
+TArray<TWeakObjectPtr<UNiagaraGeometryCacheRendererProperties>> UNiagaraGeometryCacheRendererProperties::GeometryCacheRendererPropertiesToDeferredInit;
+FNiagaraVariable UNiagaraGeometryCacheRendererProperties::Particles_GeoCacheRotation;
+FNiagaraVariable UNiagaraGeometryCacheRendererProperties::Particles_Age;
+FNiagaraVariable UNiagaraGeometryCacheRendererProperties::Particles_GeoCacheIsEnabled;
+FNiagaraVariable UNiagaraGeometryCacheRendererProperties::Particles_ArrayIndex;
+
+FNiagaraGeometryCacheReference::FNiagaraGeometryCacheReference()
+	: GeometryCache(nullptr)
+	, GeometryCacheUserParamBinding(FNiagaraTypeDefinition(UObject::StaticClass()))
+{
+}
+
+UNiagaraGeometryCacheRendererProperties::UNiagaraGeometryCacheRendererProperties()
+{
+	AttributeBindings.Reserve(6);
+	AttributeBindings.Add(&PositionBinding);
+	AttributeBindings.Add(&RotationBinding);
+	AttributeBindings.Add(&ScaleBinding);
+	AttributeBindings.Add(&ElapsedTimeBinding);
+	AttributeBindings.Add(&EnabledBinding);
+	AttributeBindings.Add(&RendererVisibilityTagBinding);
+	AttributeBindings.Add(&ArrayIndexBinding);
+
+	if (GeometryCaches.Num() == 0)
+	{
+		GeometryCaches.AddDefaulted();
+	}
+	GeometryCaches.Shrink();
+}
+
+void UNiagaraGeometryCacheRendererProperties::PostLoad()
+{
+	Super::PostLoad();
+	PostLoadBindings(ENiagaraRendererSourceDataMode::Particles);
+}
+
+template<typename T>
+FNiagaraVariableAttributeBinding CreateDefaultBinding(FNiagaraVariable DefaultValue, const T& DefaultValueData)
+{
+	DefaultValue.SetValue(DefaultValueData);
+	FNiagaraVariableAttributeBinding Binding;
+	Binding.Setup(DefaultValue, DefaultValue);
+	return Binding;
+}
+
+void UNiagaraGeometryCacheRendererProperties::InitBindings()
+{
+	if (!PositionBinding.IsValid())
+	{
+		PositionBinding = FNiagaraConstants::GetAttributeDefaultBinding(SYS_PARAM_PARTICLES_POSITION);
+		ScaleBinding = FNiagaraConstants::GetAttributeDefaultBinding(SYS_PARAM_PARTICLES_SCALE);
+		RendererVisibilityTagBinding = FNiagaraConstants::GetAttributeDefaultBinding(SYS_PARAM_PARTICLES_VISIBILITY_TAG);
+		ArrayIndexBinding = CreateDefaultBinding(Particles_ArrayIndex, 0);
+		RotationBinding = CreateDefaultBinding(Particles_GeoCacheRotation, FVector3f(0, 0, 0));
+		ElapsedTimeBinding = CreateDefaultBinding(Particles_Age, 0.0f);
+		EnabledBinding = CreateDefaultBinding(Particles_GeoCacheIsEnabled, false);
+	}
+}
+
+void UNiagaraGeometryCacheRendererProperties::InitDefaultAttributes()
+{
+	if (!Particles_GeoCacheRotation.IsValid())
+	{
+		Particles_GeoCacheRotation = FNiagaraVariable(FNiagaraTypeDefinition::GetVec3Def(), TEXT("Particles.GeoCacheRotation"));
+	}
+	if (!Particles_Age.IsValid())
+	{
+		Particles_Age = FNiagaraVariable(FNiagaraTypeDefinition::GetFloatDef(), TEXT("Particles.Age"));
+	}
+	if (!Particles_GeoCacheIsEnabled.IsValid())
+	{
+		Particles_GeoCacheIsEnabled = FNiagaraVariable(FNiagaraTypeDefinition::GetBoolDef(), TEXT("Particles.GeoCacheIsEnabled"));
+	}
+	if (!Particles_ArrayIndex.IsValid())
+	{
+		Particles_ArrayIndex = FNiagaraVariable(FNiagaraTypeDefinition::GetIntDef(), TEXT("Particles.GeoCacheArrayIndex"));
+	}
+}
+
+void UNiagaraGeometryCacheRendererProperties::PostInitProperties()
+{
+	Super::PostInitProperties();
+
+	if (HasAnyFlags(RF_ClassDefaultObject) == false)
+	{
+		// We can end up hitting PostInitProperties before the Niagara Module has initialized bindings this needs, mark this object for deferred init and early out.
+		if (FModuleManager::Get().IsModuleLoaded("Niagara") == false)
+		{
+			GeometryCacheRendererPropertiesToDeferredInit.Add(this);
+		}
+		else
+		{
+			InitDefaultAttributes();
+			
+			if (!PositionBinding.IsValid())
+			{
+				PositionBinding = FNiagaraConstants::GetAttributeDefaultBinding(SYS_PARAM_PARTICLES_POSITION);
+			}
+			if (!ScaleBinding.IsValid())
+			{
+				ScaleBinding = FNiagaraConstants::GetAttributeDefaultBinding(SYS_PARAM_PARTICLES_SCALE);
+			}
+			if (!RendererVisibilityTagBinding.IsValid())
+			{
+				RendererVisibilityTagBinding = FNiagaraConstants::GetAttributeDefaultBinding(SYS_PARAM_PARTICLES_VISIBILITY_TAG);
+			}
+			if (!RotationBinding.IsValid())
+			{
+				RotationBinding = CreateDefaultBinding(Particles_GeoCacheRotation, FVector3f(0, 0, 0));
+			}
+			if (!ElapsedTimeBinding.IsValid())
+			{
+				ElapsedTimeBinding = CreateDefaultBinding(Particles_Age, 0.0f);
+			}
+			if (!EnabledBinding.IsValid())
+			{
+				EnabledBinding = CreateDefaultBinding(Particles_GeoCacheIsEnabled, false);
+			}
+			if (!ArrayIndexBinding.IsValid())
+			{
+				ArrayIndexBinding = CreateDefaultBinding(Particles_ArrayIndex, 0);
+			}
+		}
+	}
+}
+
+void UNiagaraGeometryCacheRendererProperties::InitCDOPropertiesAfterModuleStartup()
+{
+	InitDefaultAttributes();
+	
+	UNiagaraGeometryCacheRendererProperties* CDO = CastChecked<UNiagaraGeometryCacheRendererProperties>(StaticClass()->GetDefaultObject());
+	CDO->PositionBinding = FNiagaraConstants::GetAttributeDefaultBinding(SYS_PARAM_PARTICLES_POSITION);
+	CDO->ScaleBinding = FNiagaraConstants::GetAttributeDefaultBinding(SYS_PARAM_PARTICLES_SCALE);
+	CDO->RendererVisibilityTagBinding = FNiagaraConstants::GetAttributeDefaultBinding(SYS_PARAM_PARTICLES_VISIBILITY_TAG);
+	CDO->RotationBinding = CreateDefaultBinding(Particles_GeoCacheRotation, FVector3f(0, 0, 0));
+	CDO->ElapsedTimeBinding = CreateDefaultBinding(Particles_Age, 0.0f);
+	CDO->EnabledBinding = CreateDefaultBinding(Particles_GeoCacheIsEnabled, false);
+	CDO->ArrayIndexBinding = CreateDefaultBinding(Particles_ArrayIndex, 0);
+
+	for (TWeakObjectPtr<UNiagaraGeometryCacheRendererProperties>& WeakGeometryCacheRendererProperties : GeometryCacheRendererPropertiesToDeferredInit)
+	{
+		if (WeakGeometryCacheRendererProperties.Get())
+		{
+			if (!WeakGeometryCacheRendererProperties->PositionBinding.IsValid())
+			{
+				WeakGeometryCacheRendererProperties->PositionBinding = FNiagaraConstants::GetAttributeDefaultBinding(SYS_PARAM_PARTICLES_POSITION);
+			}
+			if (!WeakGeometryCacheRendererProperties->ScaleBinding.IsValid())
+			{
+				WeakGeometryCacheRendererProperties->ScaleBinding = FNiagaraConstants::GetAttributeDefaultBinding(SYS_PARAM_PARTICLES_SCALE);
+			}
+			if (!WeakGeometryCacheRendererProperties->RendererVisibilityTagBinding.IsValid())
+			{
+				WeakGeometryCacheRendererProperties->RendererVisibilityTagBinding = FNiagaraConstants::GetAttributeDefaultBinding(SYS_PARAM_PARTICLES_VISIBILITY_TAG);
+			}
+			if (!WeakGeometryCacheRendererProperties->RotationBinding.IsValid())
+			{
+				WeakGeometryCacheRendererProperties->RotationBinding = CreateDefaultBinding(Particles_GeoCacheRotation, FVector3f(0, 0, 0));
+			}
+			if (!WeakGeometryCacheRendererProperties->ElapsedTimeBinding.IsValid())
+			{
+				WeakGeometryCacheRendererProperties->ElapsedTimeBinding = CreateDefaultBinding(Particles_Age, 0.0f);
+			}
+			if (!WeakGeometryCacheRendererProperties->EnabledBinding.IsValid())
+			{
+				WeakGeometryCacheRendererProperties->EnabledBinding = CreateDefaultBinding(Particles_GeoCacheIsEnabled, false);
+			}
+			if (!WeakGeometryCacheRendererProperties->ArrayIndexBinding.IsValid())
+			{
+				WeakGeometryCacheRendererProperties->ArrayIndexBinding = CreateDefaultBinding(Particles_ArrayIndex, 0);
+			}
+		}
+	}
+}
+
+FNiagaraRenderer* UNiagaraGeometryCacheRendererProperties::CreateEmitterRenderer(ERHIFeatureLevel::Type FeatureLevel, const FNiagaraEmitterInstance* Emitter, const FNiagaraSystemInstanceController& InController)
+{
+	FNiagaraRenderer* NewRenderer = new FNiagaraRendererGeometryCache(FeatureLevel, this, Emitter);
+	NewRenderer->Initialize(this, Emitter, InController);
+	return NewRenderer;
+}
+
+void UNiagaraGeometryCacheRendererProperties::GetUsedMaterials(const FNiagaraEmitterInstance* InEmitter, TArray<UMaterialInterface*>& MaterialInterfaces) const
+{
+	for (const FNiagaraGeometryCacheReference& Entry : GeometryCaches)
+	{
+		if (Entry.GeometryCache)
+		{
+			int32 MaxIndex = FMath::Max(Entry.GeometryCache->Materials.Num(), Entry.OverrideMaterials.Num());
+			for (int i = 0; i < MaxIndex; i++)
+			{
+				if (Entry.OverrideMaterials.IsValidIndex(i) && Entry.OverrideMaterials[i])
+				{
+					MaterialInterfaces.Add(Entry.OverrideMaterials[i]);
+				}
+				else if (Entry.GeometryCache->Materials.IsValidIndex(i))
+				{
+					MaterialInterfaces.Add(Entry.GeometryCache->Materials[i]);
+				}
+			}
+		}
+	}
+}
+
+bool UNiagaraGeometryCacheRendererProperties::PopulateRequiredBindings(FNiagaraParameterStore& InParameterStore)
+{
+	bool bAnyAdded = false;
+
+	for (const FNiagaraVariableAttributeBinding* Binding : AttributeBindings)
+	{
+		if (Binding && Binding->CanBindToHostParameterMap())
+		{
+			InParameterStore.AddParameter(Binding->GetParamMapBindableVariable(), false);
+			bAnyAdded = true;
+		}
+	}
+
+	for (const FNiagaraGeometryCacheReference& Entry : GeometryCaches)
+	{
+		FNiagaraVariable Variable = Entry.GeometryCacheUserParamBinding.Parameter;
+		if (Variable.IsValid())
+		{
+			InParameterStore.AddParameter(Variable, false);
+			bAnyAdded = true;
+		}
+	}
+
+	return bAnyAdded;
+}
+
+#if WITH_EDITORONLY_DATA
+
+const FSlateBrush* UNiagaraGeometryCacheRendererProperties::GetStackIcon() const
+{
+	return FSlateIconFinder::FindIconBrushForClass(GetClass());
+}
+
+void UNiagaraGeometryCacheRendererProperties::GetRendererWidgets(const FNiagaraEmitterInstance* InEmitter, TArray<TSharedPtr<SWidget>>& OutWidgets, TSharedPtr<FAssetThumbnailPool> InThumbnailPool) const
+{
+	TSharedRef<SWidget> DefaultThumbnailWidget = SNew(SImage)
+		.Image(FSlateIconFinder::FindIconBrushForClass(StaticClass()));
+
+	int32 ThumbnailSize = 32;
+	for(const FNiagaraGeometryCacheReference& Entry : GeometryCaches)
+	{
+		TSharedPtr<SWidget> ThumbnailWidget = DefaultThumbnailWidget;
+
+		UGeometryCache* GeometryCache = Entry.GeometryCache;
+		if (GeometryCache)
+		{
+			TSharedPtr<FAssetThumbnail> AssetThumbnail = MakeShareable(new FAssetThumbnail(GeometryCache, ThumbnailSize, ThumbnailSize, InThumbnailPool));
+			ThumbnailWidget = AssetThumbnail->MakeThumbnailWidget();
+		}
+		
+		OutWidgets.Add(ThumbnailWidget);
+	}
+
+	if (GeometryCaches.Num() == 0)
+	{
+		OutWidgets.Add(DefaultThumbnailWidget);
+	}
+}
+
+void UNiagaraGeometryCacheRendererProperties::GetRendererTooltipWidgets(const FNiagaraEmitterInstance* InEmitter, TArray<TSharedPtr<SWidget>>& OutWidgets, TSharedPtr<FAssetThumbnailPool> AssetThumbnailPool) const
+{
+	TSharedRef<SWidget> DefaultGeoCacheTooltip = SNew(STextBlock)
+		.Text(LOCTEXT("GeoCacheRenderer", "Geometry Cache Renderer"));
+	
+	TArray<TSharedPtr<SWidget>> RendererWidgets;
+	if (GeometryCaches.Num() > 0)
+	{
+		GetRendererWidgets(InEmitter, RendererWidgets, AssetThumbnailPool);
+	}
+	
+	for(int32 Index = 0; Index < GeometryCaches.Num(); Index++)
+	{
+		const FNiagaraGeometryCacheReference& Entry = GeometryCaches[Index];
+		
+		TSharedPtr<SWidget> TooltipWidget = DefaultGeoCacheTooltip;		
+		// we make sure to reuse the asset widget as a thumbnail if the geometry cache is valid
+		if(Entry.GeometryCache)
+		{
+			TooltipWidget = RendererWidgets[Index];
+		}
+
+		// we override the previous thumbnail tooltip with a text indicating user parameter binding, if it exists
+		if(Entry.GeometryCacheUserParamBinding.Parameter.IsValid())
+		{
+			TooltipWidget = SNew(STextBlock)
+				.Text(FText::Format(LOCTEXT("GeoCacheBoundTooltip", "Geometry cache slot is bound to user parameter {0}"), FText::FromName(Entry.GeometryCacheUserParamBinding.Parameter.GetName())));
+		}
+		
+		OutWidgets.Add(TooltipWidget);
+	}
+
+	if (GeometryCaches.Num() == 0)
+	{
+		OutWidgets.Add(DefaultGeoCacheTooltip);
+	}
+}
+
+const TArray<FNiagaraVariable>& UNiagaraGeometryCacheRendererProperties::GetOptionalAttributes()
+{
+	static TArray<FNiagaraVariable> Attrs;
+	if (Attrs.Num() == 0)
+	{
+		Attrs.Add(SYS_PARAM_PARTICLES_POSITION);
+		Attrs.Add(SYS_PARAM_PARTICLES_SCALE);
+		Attrs.Add(SYS_PARAM_PARTICLES_VISIBILITY_TAG);
+		Attrs.Add(Particles_GeoCacheRotation);
+		Attrs.Add(Particles_Age);
+		Attrs.Add(Particles_GeoCacheIsEnabled);
+		Attrs.Add(Particles_ArrayIndex);
+	}
+	return Attrs;
+}
+
+bool UNiagaraGeometryCacheRendererProperties::IsMaterialValidForRenderer(UMaterial* Material, FText& InvalidMessage)
+{
+	if (Material->bUsedWithGeometryCache == false)
+	{
+		InvalidMessage = NSLOCTEXT("NiagaraGeometryCacheRendererProperties", "InvalidMaterialMessage", "The material isn't marked as \"Used with Geometry Cache\"");
+		return false;
+	}
+	return true;
+}
+
+void UNiagaraGeometryCacheRendererProperties::FixMaterial(UMaterial* Material)
+{
+	Material->Modify();
+	Material->bUsedWithGeometryCache = true;
+	Material->ForceRecompileForRendering();
+}
+
+void UNiagaraGeometryCacheRendererProperties::GetRendererFeedback(const UNiagaraEmitter* InEmitter,	TArray<FText>& OutErrors, TArray<FText>& OutWarnings, TArray<FText>& OutInfo) const
+{
+	OutInfo.Add(FText::FromString(TEXT("The geometry cache renderer is still an experimental feature.")));
+}
+
+#endif
+
+#undef LOCTEXT_NAMESPACE
\ No newline at end of file
diff --git a/Engine/Plugins/Experimental/GeometryCache/Source/GeometryCache/Private/NiagaraRendererGeometryCache.cpp b/Engine/Plugins/Experimental/GeometryCache/Source/GeometryCache/Private/NiagaraRendererGeometryCache.cpp
new file mode 100644
--- /dev/null
+++ b/Engine/Plugins/Experimental/GeometryCache/Source/GeometryCache/Private/NiagaraRendererGeometryCache.cpp
@@ -0,0 +1,423 @@
+// Copyright Epic Games, Inc. All Rights Reserved.
+
+#include "NiagaraRendererGeometryCache.h"
+
+#include "GeometryCacheStreamingManager.h"
+#include "NiagaraConstants.h"
+#include "NiagaraDataSet.h"
+#include "NiagaraGeometryCacheRendererProperties.h"
+#include "NiagaraSystemInstance.h"
+#include "Async/Async.h"
+
+static float GNiagaraGeometryComponentRenderPoolInactiveTimeLimit = 5;
+static FAutoConsoleVariableRef CVarNiagaraComponentRenderPoolInactiveTimeLimit(
+	TEXT("fx.Niagara.GeometryComponentRenderPoolInactiveTimeLimit"),
+	GNiagaraGeometryComponentRenderPoolInactiveTimeLimit,
+	TEXT("The time in seconds an inactive component can linger in the pool before being destroyed."),
+	ECVF_Default
+);
+
+//////////////////////////////////////////////////////////////////////////
+
+
+FNiagaraRendererGeometryCache::FNiagaraRendererGeometryCache(ERHIFeatureLevel::Type FeatureLevel, const UNiagaraRendererProperties *InProps, const FNiagaraEmitterInstance* Emitter)
+	: FNiagaraRenderer(FeatureLevel, InProps, Emitter)
+{
+	const UNiagaraGeometryCacheRendererProperties* Properties = CastChecked<const UNiagaraGeometryCacheRendererProperties>(InProps);
+	ComponentPool.Reserve(Properties->ComponentCountLimit);
+}
+
+FNiagaraRendererGeometryCache::~FNiagaraRendererGeometryCache()
+{
+	// These should have been freed in DestroyRenderState_Concurrent
+	check(ComponentPool.Num() == 0);
+}
+
+void FNiagaraRendererGeometryCache::DestroyRenderState_Concurrent()
+{
+	// Rendering resources are being destroyed, but the component pool and their owner actor must be destroyed on the game thread
+	AsyncTask(
+		ENamedThreads::GameThread,
+		[Pool_GT = MoveTemp(ComponentPool), Owner_GT = MoveTemp(SpawnedOwner)]()
+		{
+			for (const FComponentPoolEntry& PoolEntry : Pool_GT)
+			{
+				if (PoolEntry.Component.IsValid())
+				{
+					IGeometryCacheStreamingManager::Get().RemoveStreamingComponent(PoolEntry.Component.Get());
+					PoolEntry.Component->DestroyComponent();
+				}
+			}
+
+			if (AActor* OwnerActor = Owner_GT.Get())
+			{
+				OwnerActor->Destroy();
+			}
+		}
+	);
+	SpawnedOwner.Reset();
+}
+
+UGeometryCache* ResolveGeometryCache(const FNiagaraGeometryCacheReference& Entry, const FNiagaraEmitterInstance* Emitter)
+{
+	UGeometryCache* FoundCache = nullptr;
+
+	FNiagaraVariable Variable = Entry.GeometryCacheUserParamBinding.Parameter;
+	if (Variable.IsValid() && Emitter)
+	{
+		UGeometryCache* GeometryCache = Cast<UGeometryCache>(Emitter->GetRendererBoundVariables().GetUObject(Variable));
+		if (GeometryCache)
+		{
+			FoundCache = GeometryCache;
+		}
+	}
+
+	if (!FoundCache)
+	{
+		FoundCache = Entry.GeometryCache;
+	}
+
+	return FoundCache;
+}
+
+/** Update render data buffer from attributes */
+void FNiagaraRendererGeometryCache::PostSystemTick_GameThread(const UNiagaraRendererProperties* InProperties, const FNiagaraEmitterInstance* Emitter)
+{
+	FNiagaraSystemInstance* SystemInstance = Emitter->GetParentSystemInstance();
+
+	//Bail if we don't have the required attributes to render this emitter.
+	const UNiagaraGeometryCacheRendererProperties* Properties = CastChecked<const UNiagaraGeometryCacheRendererProperties>(InProperties);
+	if (!SystemInstance || !Properties || Properties->GeometryCaches.Num() == 0 || SimTarget == ENiagaraSimTarget::GPUComputeSim)
+	{
+		return;
+	}
+
+	USceneComponent* AttachComponent = SystemInstance->GetAttachComponent();
+	if (!AttachComponent)
+	{
+		// we can't attach the components anywhere, so just bail
+		return;
+	}
+
+	FNiagaraDataSet& Data = Emitter->GetData();
+	FNiagaraDataBuffer& ParticleData = Data.GetCurrentDataChecked();
+	FNiagaraDataSetReaderInt32<FNiagaraBool> EnabledAccessor = FNiagaraDataSetAccessor<FNiagaraBool>::CreateReader(Data, Properties->EnabledBinding.GetDataSetBindableVariable().GetName());
+	FNiagaraDataSetReaderInt32<int32> VisTagAccessor = FNiagaraDataSetAccessor<int32>::CreateReader(Data, Properties->RendererVisibilityTagBinding.GetDataSetBindableVariable().GetName());
+	FNiagaraDataSetReaderInt32<int32> ArrayIndexAccessor = FNiagaraDataSetAccessor<int32>::CreateReader(Data, Properties->ArrayIndexBinding.GetDataSetBindableVariable().GetName());
+	FNiagaraDataSetReaderInt32<int32> UniqueIDAccessor = FNiagaraDataSetAccessor<int32>::CreateReader(Data, FName("UniqueID"));
+	float CurrentTime = AttachComponent->GetWorld()->GetRealTimeSeconds();
+	
+	auto IsParticleEnabled = [&EnabledAccessor, &VisTagAccessor, Properties](int32 ParticleIndex)
+	{
+		if (EnabledAccessor.GetSafe(ParticleIndex, true))
+		{
+			if (VisTagAccessor.IsValid())
+			{
+				return VisTagAccessor.GetSafe(ParticleIndex, 0) == Properties->RendererVisibility;
+			}
+			return true;
+		}
+		return false;
+	};
+
+	TMap<int32, int32> ParticlesWithComponents;
+	TArray<int32> FreeList;
+	if (Properties->bAssignComponentsOnParticleID && ComponentPool.Num() > 0)
+	{
+		FreeList.Reserve(ComponentPool.Num());
+
+		// Determine the slots that were assigned to particles last frame
+		TMap<int32, int32> UsedSlots;
+		UsedSlots.Reserve(ComponentPool.Num());
+		for (int32 EntryIndex = 0; EntryIndex < ComponentPool.Num(); ++EntryIndex)
+		{
+			FComponentPoolEntry& Entry = ComponentPool[EntryIndex];
+			if (Entry.LastAssignedToParticleID >= 0)
+			{
+				UsedSlots.Emplace(Entry.LastAssignedToParticleID, EntryIndex);
+			}
+			else
+			{
+				FreeList.Add(EntryIndex);
+			}
+		}
+
+		// Ensure the final list only contains particles that are alive and enabled
+		ParticlesWithComponents.Reserve(UsedSlots.Num());
+		for (uint32 ParticleIndex = 0; ParticleIndex < ParticleData.GetNumInstances(); ParticleIndex++)
+		{
+			int32 ParticleID = UniqueIDAccessor.GetSafe(ParticleIndex, -1);
+			int32 PoolIndex;
+			if (UsedSlots.RemoveAndCopyValue(ParticleID, PoolIndex))
+			{
+				if (IsParticleEnabled(ParticleIndex))
+				{
+					ParticlesWithComponents.Emplace(ParticleID, PoolIndex);
+				}
+				else
+				{
+					// Particle has disabled components since last tick, ensure the component for this entry gets deactivated before re-use
+					USceneComponent* Component = ComponentPool[PoolIndex].Component.Get();
+					if (Component && Component->IsActive())
+					{
+						Component->Deactivate();
+						Component->SetVisibility(false, true);
+					}
+					FreeList.Add(PoolIndex);
+					ComponentPool[PoolIndex].LastAssignedToParticleID = -1;
+				}
+			}
+		}
+
+		// Any remaining in the used slots are now free to be reclaimed, due to their particles either dying or having their component disabled
+		for (TPair<int32, int32> UsedSlot : UsedSlots)
+		{
+			// Particle has died since last tick, ensure the component for this entry gets deactivated before re-use
+			USceneComponent* Component = ComponentPool[UsedSlot.Value].Component.Get();
+			if (Component && Component->IsActive())
+			{
+				Component->Deactivate();
+				Component->SetVisibility(false, true);
+			}
+			FreeList.Add(UsedSlot.Value);
+			ComponentPool[UsedSlot.Value].LastAssignedToParticleID = -1;
+		}
+	}
+
+	const int32 MaxComponents = Properties->ComponentCountLimit;
+	int32 ComponentCount = 0;
+	for (uint32 ParticleIndex = 0; ParticleIndex < ParticleData.GetNumInstances(); ParticleIndex++)
+	{
+		if (!IsParticleEnabled(ParticleIndex))
+		{
+			// Skip particles that don't want a component
+			continue;
+		}
+
+		int32 ParticleID = -1;
+		int32 PoolIndex = -1;
+		if (Properties->bAssignComponentsOnParticleID)
+		{
+			// Get the particle ID and see if we have any components already assigned to the particle
+			ParticleID = UniqueIDAccessor.GetSafe(ParticleIndex, -1);
+			ParticlesWithComponents.RemoveAndCopyValue(ParticleID, PoolIndex);
+		}
+
+		if (PoolIndex == -1 && ComponentCount + ParticlesWithComponents.Num() >= MaxComponents)
+		{
+			// The pool is full and there aren't any unused slots to claim
+			continue;
+		}
+
+		// Acquire a component for this particle
+		UGeometryCacheComponent* GeometryComponent = nullptr;
+		if (PoolIndex == -1)
+		{
+			// Start by trying to pull from the pool
+			if (!Properties->bAssignComponentsOnParticleID)
+			{
+				// We can just take the next slot
+				PoolIndex = ComponentCount < ComponentPool.Num() ? ComponentCount : -1;
+			}
+			else if (FreeList.Num())
+			{
+				PoolIndex = FreeList.Pop(false);
+			}
+		}
+
+		if (PoolIndex >= 0)
+		{
+			GeometryComponent = ComponentPool[PoolIndex].Component.Get();
+		}
+
+		if (!GeometryComponent || GeometryComponent->HasAnyFlags(RF_BeginDestroyed | RF_FinishDestroyed))
+		{
+			int32 CacheIndex = ArrayIndexAccessor.GetSafe(ParticleIndex, INDEX_NONE);
+			if (CacheIndex == INDEX_NONE && Emitter->GetCachedEmitter()->bDeterminism)
+			{
+				int32 Seed = Properties->bAssignComponentsOnParticleID ? ParticleID : ParticleIndex;
+				FRandomStream RandomStream = FRandomStream(Seed * 907633515U); // multiply the seed, otherwise we get very poor randomness for small seeds
+				CacheIndex = RandomStream.RandRange(0, Properties->GeometryCaches.Num());
+			}
+			else if (CacheIndex == INDEX_NONE)
+			{
+				CacheIndex = FMath::RandRange(0, Properties->GeometryCaches.Num() - 1);
+			}
+			CacheIndex = CacheIndex % Properties->GeometryCaches.Num(); // wrap around if index is too large
+			const FNiagaraGeometryCacheReference& Entry = Properties->GeometryCaches[CacheIndex];
+			UGeometryCache* Cache = ResolveGeometryCache(Entry, Emitter);
+			if (!Cache)
+			{
+				// no valid geometry cache was provided for this particle
+				continue;
+			}
+
+			// Determine the owner actor or spawn one
+			AActor* OwnerActor = SpawnedOwner.Get();
+			if (OwnerActor == nullptr)
+			{
+				OwnerActor = AttachComponent->GetOwner();
+				if (OwnerActor == nullptr)
+				{
+					// NOTE: This can happen with spawned systems
+					OwnerActor = AttachComponent->GetWorld()->SpawnActor<AActor>();
+					OwnerActor->SetFlags(RF_Transient);
+					SpawnedOwner = OwnerActor;
+				}
+			}
+
+			// if we don't have a pooled component we create a new one from the template
+			GeometryComponent = NewObject<UGeometryCacheComponent>(OwnerActor);
+			IGeometryCacheStreamingManager::Get().AddStreamingComponent(GeometryComponent);
+			GeometryComponent->SetFlags(RF_Transient);
+			GeometryComponent->SetupAttachment(AttachComponent);
+			GeometryComponent->RegisterComponent();
+			GeometryComponent->AddTickPrerequisiteComponent(AttachComponent);
+			GeometryComponent->SetLooping(Properties->bIsLooping);
+			GeometryComponent->SetGeometryCache(Cache);
+			GeometryComponent->SetManualTick(true); // we want to tick the component with the delta time of the niagara sim
+
+			if (Emitter->GetCachedEmitter()->bLocalSpace)
+			{
+				GeometryComponent->SetAbsolute(false, false, false);
+			}
+			else
+			{
+				GeometryComponent->SetAbsolute(true, true, true);
+			}
+
+			// set the override materials
+			for (int32 i = 0; i < Entry.OverrideMaterials.Num(); i++)
+			{
+				const TObjectPtr<UMaterialInterface>& OverrideMat = Entry.OverrideMaterials[i];
+				GeometryComponent->SetMaterial(i, OverrideMat);
+			}
+
+			if (PoolIndex >= 0)
+			{
+				// This should only happen if the component was destroyed externally
+				ComponentPool[PoolIndex].Component = GeometryComponent;
+				ComponentPool[PoolIndex].LastElapsedTime = 0;
+			}
+			else
+			{
+				// Add a new pool entry
+				PoolIndex = ComponentPool.Num();
+				ComponentPool.AddDefaulted_GetRef().Component = GeometryComponent;
+				ComponentPool[PoolIndex].LastElapsedTime = 0;
+			}
+		}
+
+		const auto PositionAccessor = FNiagaraDataSetAccessor<FVector3f>::CreateReader(Data, Properties->PositionBinding.GetDataSetBindableVariable().GetName());
+		const auto RotationAccessor = FNiagaraDataSetAccessor<FVector3f>::CreateReader(Data, Properties->RotationBinding.GetDataSetBindableVariable().GetName());
+		const auto ScaleAccessor = FNiagaraDataSetAccessor<FVector3f>::CreateReader(Data, Properties->ScaleBinding.GetDataSetBindableVariable().GetName());
+		const auto ElapsedTimeAccessor = FNiagaraDataSetAccessor<float>::CreateReader(Data, Properties->ElapsedTimeBinding.GetDataSetBindableVariable().GetName());
+
+		FVector Position = PositionAccessor.GetSafe(ParticleIndex, FVector3f::ZeroVector);
+		FVector Scale = ScaleAccessor.GetSafe(ParticleIndex, FVector3f::OneVector);
+		FVector RotationVector = RotationAccessor.GetSafe(ParticleIndex, FVector3f::ZeroVector);
+		FTransform Transform(FRotator(RotationVector.X, RotationVector.Y, RotationVector.Z), Position, Scale);
+		GeometryComponent->SetRelativeTransform(Transform);
+
+		// activate the component
+		if (!GeometryComponent->IsActive())
+		{
+			GeometryComponent->SetVisibility(true, true);
+			GeometryComponent->Activate(false);
+		}
+
+		float ElapsedTime = ElapsedTimeAccessor.GetSafe(ParticleIndex, 0);
+
+		float Duration = GeometryComponent->GetDuration();
+		if (Properties->bIsLooping && ElapsedTime < 0)
+		{
+			ElapsedTime -= FMath::CeilToFloat(ElapsedTime / Duration) * Duration;
+		}
+		else if (Properties->bIsLooping && ElapsedTime > Duration)
+		{
+			ElapsedTime -= FMath::Floor(ElapsedTime / Duration) * Duration;
+		}
+
+		FComponentPoolEntry& PoolEntry = ComponentPool[PoolIndex];
+		GeometryComponent->TickAtThisTime(ElapsedTime, true, ElapsedTime < PoolEntry.LastElapsedTime, Properties->bIsLooping);
+
+		PoolEntry.LastElapsedTime = ElapsedTime;
+		PoolEntry.LastAssignedToParticleID = ParticleID;
+		PoolEntry.LastActiveTime = CurrentTime;
+
+		++ComponentCount;
+		if (ComponentCount >= MaxComponents)
+		{
+			// We've hit our prescribed limit
+			break;
+		}
+	}
+
+	if (ComponentCount < ComponentPool.Num())
+	{
+		// go over the pooled components we didn't need this tick to see if we can destroy some and deactivate the rest
+		for (int32 PoolIndex = 0; PoolIndex < ComponentPool.Num(); ++PoolIndex)
+		{
+			FComponentPoolEntry& PoolEntry = ComponentPool[PoolIndex];
+			if (Properties->bAssignComponentsOnParticleID)
+			{
+				if (PoolEntry.LastAssignedToParticleID >= 0)
+				{
+					// This one's in use
+					continue;
+				}
+			}
+			else if (PoolIndex < ComponentCount)
+			{
+				continue;
+			}
+
+			USceneComponent* Component = PoolEntry.Component.Get();
+			if (!Component || (CurrentTime - PoolEntry.LastActiveTime) >= GNiagaraGeometryComponentRenderPoolInactiveTimeLimit)
+			{
+				if (Component)
+				{
+					IGeometryCacheStreamingManager::Get().RemoveStreamingComponent(PoolEntry.Component.Get());
+					Component->DestroyComponent();
+				}
+
+				// destroy the component pool slot
+				ComponentPool.RemoveAtSwap(PoolIndex, 1, false);
+				--PoolIndex;
+				continue;
+			}
+			else if (Component->IsActive())
+			{
+				Component->Deactivate();
+				Component->SetVisibility(false, true);
+			}
+		}
+	}
+}
+
+void FNiagaraRendererGeometryCache::OnSystemComplete_GameThread(const UNiagaraRendererProperties* InProperties, const FNiagaraEmitterInstance* Emitter)
+{
+	ResetComponentPool(true);
+}
+
+void FNiagaraRendererGeometryCache::ResetComponentPool(bool bResetOwner)
+{
+	for (FComponentPoolEntry& PoolEntry : ComponentPool)
+	{
+		if (PoolEntry.Component.IsValid())
+		{
+			IGeometryCacheStreamingManager::Get().RemoveStreamingComponent(PoolEntry.Component.Get());
+			PoolEntry.Component->DestroyComponent();
+		}
+	}
+	ComponentPool.SetNum(0, false);
+
+	if (bResetOwner)
+	{
+		if (AActor* OwnerActor = SpawnedOwner.Get())
+		{
+			SpawnedOwner.Reset();
+			OwnerActor->Destroy();
+		}
+	}
+}
diff --git a/Engine/Plugins/FX/Niagara/Source/Niagara/Public/NiagaraRenderer.h b/Engine/Plugins/FX/Niagara/Source/Niagara/Public/NiagaraRenderer.h
--- a/Engine/Plugins/FX/Niagara/Source/Niagara/Public/NiagaraRenderer.h
+++ b/Engine/Plugins/FX/Niagara/Source/Niagara/Public/NiagaraRenderer.h
@@ -70,7 +70,7 @@ struct FParticleRenderData
 /**
 * Base class for Niagara System renderers.
 */
-class FNiagaraRenderer
+class NIAGARA_API FNiagaraRenderer
 {
 public:
 
@@ -116,17 +116,17 @@ public:
 	FORCEINLINE static FRHIShaderResourceView* GetSrvOrDefaultUInt(FGlobalDynamicReadBuffer::FAllocation& Allocation) { return Allocation.IsValid() ? (FRHIShaderResourceView*)Allocation.SRV : GetDummyUIntBuffer(); }
 	FORCEINLINE static FRHIShaderResourceView* GetSrvOrDefaultInt(FGlobalDynamicReadBuffer::FAllocation& Allocation) { return Allocation.IsValid() ? (FRHIShaderResourceView*)Allocation.SRV : GetDummyIntBuffer(); }
 
-	NIAGARA_API static FRHIShaderResourceView* GetDummyFloatBuffer();
-	NIAGARA_API static FRHIShaderResourceView* GetDummyFloat2Buffer();
-	NIAGARA_API static FRHIShaderResourceView* GetDummyFloat4Buffer();
-	NIAGARA_API static FRHIShaderResourceView* GetDummyWhiteColorBuffer();
-	NIAGARA_API static FRHIShaderResourceView* GetDummyIntBuffer();
-	NIAGARA_API static FRHIShaderResourceView* GetDummyUIntBuffer();
-	NIAGARA_API static FRHIShaderResourceView* GetDummyUInt4Buffer();
-	NIAGARA_API static FRHIShaderResourceView* GetDummyTextureReadBuffer2D();
-	NIAGARA_API static FRHIShaderResourceView* GetDummyTextureReadBuffer2DArray();
-	NIAGARA_API static FRHIShaderResourceView* GetDummyTextureReadBuffer3D();
-	NIAGARA_API static FRHIShaderResourceView* GetDummyHalfBuffer();
+	static FRHIShaderResourceView* GetDummyFloatBuffer();
+	static FRHIShaderResourceView* GetDummyFloat2Buffer();
+	static FRHIShaderResourceView* GetDummyFloat4Buffer();
+	static FRHIShaderResourceView* GetDummyWhiteColorBuffer();
+	static FRHIShaderResourceView* GetDummyIntBuffer();
+	static FRHIShaderResourceView* GetDummyUIntBuffer();
+	static FRHIShaderResourceView* GetDummyUInt4Buffer();
+	static FRHIShaderResourceView* GetDummyTextureReadBuffer2D();
+	static FRHIShaderResourceView* GetDummyTextureReadBuffer2DArray();
+	static FRHIShaderResourceView* GetDummyTextureReadBuffer3D();
+	static FRHIShaderResourceView* GetDummyHalfBuffer();
 
 	FORCEINLINE ENiagaraSimTarget GetSimTarget() const { return SimTarget; }
 
