Fixed the "Skip Computing Tangents" Alembic import option to be applicable to GeometryCache only. This option is now only visible on the GeometryCacheAbcFileComponent as originally intended.

diff --git a/Engine/Plugins/Experimental/AlembicImporter/Source/AlembicLibrary/Private/AbcPolyMesh.cpp b/Engine/Plugins/Experimental/AlembicImporter/Source/AlembicLibrary/Private/AbcPolyMesh.cpp
--- a/Engine/Plugins/Experimental/AlembicImporter/Source/AlembicLibrary/Private/AbcPolyMesh.cpp
+++ b/Engine/Plugins/Experimental/AlembicImporter/Source/AlembicLibrary/Private/AbcPolyMesh.cpp
@@ -206,7 +206,9 @@ void FAbcPolyMesh::SetFrameAndTime(const float InTime, const int32 FrameIndex, c
 				}
 			}
 
-			if (!File->GetImportSettings()->NormalGenerationSettings.bSkipComputingTangents)
+			// Computing tangents skippable for GeometryCache only, the other import types need tangents
+			if (File->GetImportSettings()->ImportType != EAlembicImportType::GeometryCache ||
+				!File->GetImportSettings()->NormalGenerationSettings.bSkipComputingTangents)
 			{
 				AbcImporterUtilities::ComputeTangents(WriteSample, File->GetImportSettings()->NormalGenerationSettings.bIgnoreDegenerateTriangles, *File->GetMeshUtilities());
 			}
diff --git a/Engine/Plugins/Experimental/AlembicImporter/Source/AlembicLibrary/Public/AbcImportSettings.h b/Engine/Plugins/Experimental/AlembicImporter/Source/AlembicLibrary/Public/AbcImportSettings.h
--- a/Engine/Plugins/Experimental/AlembicImporter/Source/AlembicLibrary/Public/AbcImportSettings.h
+++ b/Engine/Plugins/Experimental/AlembicImporter/Source/AlembicLibrary/Public/AbcImportSettings.h
@@ -151,8 +151,8 @@ struct FAbcNormalGenerationSettings
 	UPROPERTY(EditAnywhere, BlueprintReadWrite, Category = NormalCalculation, meta = (EditCondition = "bRecomputeNormals"))
 	bool bIgnoreDegenerateTriangles;
 
-	/** Determines whether tangents are computed. Skipping them can improve streaming performance but may cause visual artifacts where they are required */
-	UPROPERTY(EditAnywhere, BlueprintReadWrite, Category = NormalCalculation, AdvancedDisplay)
+	/** Determines whether tangents are computed for GeometryCache. Skipping them can improve streaming performance but may cause visual artifacts where they are required */
+	UPROPERTY(EditAnywhere, BlueprintReadWrite, Category = Hidden, AdvancedDisplay)
 	bool bSkipComputingTangents;
 };
 
diff --git a/Engine/Plugins/Experimental/GeometryCache/Source/GeometryCache/Private/GeometryCachePreprocessor.cpp b/Engine/Plugins/Experimental/GeometryCache/Source/GeometryCache/Private/GeometryCachePreprocessor.cpp
--- a/Engine/Plugins/Experimental/GeometryCache/Source/GeometryCache/Private/GeometryCachePreprocessor.cpp
+++ b/Engine/Plugins/Experimental/GeometryCache/Source/GeometryCache/Private/GeometryCachePreprocessor.cpp
@@ -84,21 +84,28 @@ bool FOptimizeGeometryCachePreprocessor::AreIndexedVerticesEqual(int32 IndexBuff
 		{
 			return false;
 		}
-		// The following are already 8 bit so quantized enough we can do exact equal comparisons
-		const FPackedNormal& TangentXA = BufferedFrames[Frame].MeshData.TangentsX[VertexIndexA];
-		const FPackedNormal& TangentXB = BufferedFrames[Frame].MeshData.TangentsX[VertexIndexB];
 
-		if (TangentXA != TangentXB)
+		if (BufferedFrames[Frame].MeshData.TangentsX.Num() > 0)
 		{
-			return false;
-		}
+			// The following are already 8 bit so quantized enough we can do exact equal comparisons
+			const FPackedNormal& TangentXA = BufferedFrames[Frame].MeshData.TangentsX[VertexIndexA];
+			const FPackedNormal& TangentXB = BufferedFrames[Frame].MeshData.TangentsX[VertexIndexB];
 
-		const FPackedNormal& TangentZA = BufferedFrames[Frame].MeshData.TangentsZ[VertexIndexA];
-		const FPackedNormal& TangentZB = BufferedFrames[Frame].MeshData.TangentsZ[VertexIndexB];
+			if (TangentXA != TangentXB)
+			{
+				return false;
+			}
+		}
 
-		if (TangentZA != TangentZB)
+		if (BufferedFrames[Frame].MeshData.TangentsZ.Num() > 0)
 		{
-			return false;
+			const FPackedNormal& TangentZA = BufferedFrames[Frame].MeshData.TangentsZ[VertexIndexA];
+			const FPackedNormal& TangentZB = BufferedFrames[Frame].MeshData.TangentsZ[VertexIndexB];
+
+			if (TangentZA != TangentZB)
+			{
+				return false;
+			}
 		}
 
 		if (BufferedFrames[Frame].MeshData.Positions.Num() == BufferedFrames[Frame].MeshData.Colors.Num())
@@ -289,16 +296,29 @@ void FOptimizeGeometryCachePreprocessor::FlushBufferedFrames()
 		NewMesh.VertexInfo = OldMesh.VertexInfo;
 		NewMesh.Indices = NewIndices;
 
+		const bool bHasTangentsX = OldMesh.TangentsX.Num() > 0;
+		const bool bHasTangentsZ = OldMesh.TangentsZ.Num() > 0;
+
 		NewMesh.Positions.SetNumUninitialized(NewVerticesReordered.Num());
-		NewMesh.TangentsX.SetNumUninitialized(NewVerticesReordered.Num());
-		NewMesh.TangentsZ.SetNumUninitialized(NewVerticesReordered.Num());
+		if (bHasTangentsX)
+		{
+			NewMesh.TangentsX.SetNumUninitialized(NewVerticesReordered.Num());
+		}
+
+		if (bHasTangentsZ)
+		{
+			NewMesh.TangentsZ.SetNumUninitialized(NewVerticesReordered.Num());
+		}
+
 		if (NewMesh.VertexInfo.bHasColor0)
 		{
 			NewMesh.Colors.SetNumUninitialized(NewVerticesReordered.Num());
 		}
 		
 		if ( NewMesh.VertexInfo.bHasUV0)
+		{
 			NewMesh.TextureCoordinates.SetNumUninitialized(NewVerticesReordered.Num());
+		}
 
 		if (NewMesh.VertexInfo.bHasMotionVectors)
 		{
@@ -309,14 +329,22 @@ void FOptimizeGeometryCachePreprocessor::FlushBufferedFrames()
 		for (int32 i = 0; i < NewMesh.Positions.Num(); i++)
 		{
 			NewMesh.Positions[i] = OldMesh.Positions[NewVerticesReordered[i]];
-			NewMesh.TangentsX[i] = OldMesh.TangentsX[NewVerticesReordered[i]];
-			NewMesh.TangentsZ[i] = OldMesh.TangentsZ[NewVerticesReordered[i]];
+			if (bHasTangentsX)
+			{
+				NewMesh.TangentsX[i] = OldMesh.TangentsX[NewVerticesReordered[i]];
+			}
+			if (bHasTangentsZ)
+			{
+				NewMesh.TangentsZ[i] = OldMesh.TangentsZ[NewVerticesReordered[i]];
+			}
 			if (NewMesh.VertexInfo.bHasColor0)
 			{
 				NewMesh.Colors[i] = OldMesh.Colors[NewVerticesReordered[i]];
 			}
 			if (NewMesh.VertexInfo.bHasUV0)
+			{
 				NewMesh.TextureCoordinates[i] = OldMesh.TextureCoordinates[NewVerticesReordered[i]];
+			}
 			
 			if (NewMesh.VertexInfo.bHasMotionVectors)
 			{
