Always initialize all elemnts of the geometry cache index buffer in FGeomCacheIndexBuffer::Update()

Prevents out-of-bounds access of Geometry Cache index and vertex buffers on GPU during ray tracing acceleration structure build (which causes a GPU crash).
Also clamp FRayTracingGeometrySegment.NumPrimitives to ensure that index buffer itself is not accessed out of bounds.

This does not address the deeper underlying bug that causes a mismatch between FGeometryCacheMeshBatchInfo.NumTriangles and index buffer contents, only treats the down-stream errors.

diff --git a/Engine/Plugins/Experimental/GeometryCache/Source/GeometryCache/Private/GeometryCacheSceneProxy.cpp b/Engine/Plugins/Experimental/GeometryCache/Source/GeometryCache/Private/GeometryCacheSceneProxy.cpp
--- a/Engine/Plugins/Experimental/GeometryCache/Source/GeometryCache/Private/GeometryCacheSceneProxy.cpp
+++ b/Engine/Plugins/Experimental/GeometryCache/Source/GeometryCache/Private/GeometryCacheSceneProxy.cpp
@@ -606,6 +606,8 @@ void FGeometryCacheSceneProxy::UpdateAnimation(float NewTime, bool bNewLooping,
 				Section->RayTracingGeometry.Initializer.IndexBuffer = Section->IndexBuffer.IndexBufferRHI;
 				Section->RayTracingGeometry.Initializer.TotalPrimitiveCount = 0;
 
+				const uint32 IndexBufferNumTriangles = Section->IndexBuffer.NumIndices / 3;
+
 				TMemoryImageArray<FRayTracingGeometrySegment>& Segments = Section->RayTracingGeometry.Initializer.Segments;
 				Segments.Reset();
 
@@ -615,11 +617,19 @@ void FGeometryCacheSceneProxy::UpdateAnimation(float NewTime, bool bNewLooping,
 					FRayTracingGeometrySegment Segment;
 					Segment.FirstPrimitive = BatchInfo.StartIndex / 3;
 					Segment.NumPrimitives = BatchInfo.NumTriangles;
+
+					// Ensure that a geometry segment does not access the index buffer out of bounds
+					if (!ensureMsgf(Segment.FirstPrimitive + Segment.NumPrimitives <= IndexBufferNumTriangles, 
+						TEXT("Ray tracing geometry index buffer is smaller than what's required by FGeometryCacheMeshBatchInfo")))
+					{
+						Segment.NumPrimitives = IndexBufferNumTriangles - FMath::Min<uint32>(Segment.FirstPrimitive, IndexBufferNumTriangles);
+					}
+
 					Segment.VertexBuffer = Section->PositionBuffers[PositionBufferIndex].VertexBufferRHI;
 					Segment.MaxVertices = Section->PositionBuffers[PositionBufferIndex].GetSizeInBytes() / Segment.VertexBufferStride; // conservative estimate
 
 					Segments.Add(Segment);
-					Section->RayTracingGeometry.Initializer.TotalPrimitiveCount += BatchInfo.NumTriangles;
+					Section->RayTracingGeometry.Initializer.TotalPrimitiveCount += Segment.NumPrimitives;
 				}
 
 				if (Segments.Num() > 0)
@@ -1395,15 +1405,26 @@ void FGeomCacheIndexBuffer::Update(const TArray<uint32>& Indices)
 		bReallocate = true;
 	}
 
-		if (Indices.Num() > 0)
-		{
-			// Copy the index data into the index buffer.
-		Buffer = RHILockBuffer(IndexBufferRHI, 0, Indices.Num() * sizeof(uint32), RLM_WriteOnly);
+	if (Indices.Num() > 0)
+	{
+		// Copy the index data into the index buffer.
+		Buffer = RHILockBuffer(IndexBufferRHI, 0, NumIndices * sizeof(uint32), RLM_WriteOnly);
 	}
 
+
 	if (Buffer)
 	{
 		FMemory::Memcpy(Buffer, Indices.GetData(), Indices.Num() * sizeof(uint32));
+
+		// Do not leave any of the index buffer memory uninitialized to prevent
+		// the possibility of accessing vertex buffers out of bounds.
+		uint32* LockedIndices = reinterpret_cast<uint32*>(Buffer);
+		uint32 ValidIndexValue = Indices[0];
+		for (int32 i = Indices.Num(); i < NumIndices; ++i)
+		{
+			LockedIndices[i] = ValidIndexValue;
+		}
+
 		RHIUnlockBuffer(IndexBufferRHI);
 	}
 
