Split FGeomCacheIndexBuffer::NumIndices into explicit NumAllocatedIndices and NumValidIndices

The former corresponds to the size of the GPU memory allocation, while the latter indicates how many indices were actually uploaded.
This is used for sanity checking of FRayTracingGeometrySegment parameters, which catches potential GPU out-of-bounds access early.

diff --git a/Engine/Plugins/Experimental/GeometryCache/Source/GeometryCache/Private/GeometryCacheSceneProxy.cpp b/Engine/Plugins/Experimental/GeometryCache/Source/GeometryCache/Private/GeometryCacheSceneProxy.cpp
--- a/Engine/Plugins/Experimental/GeometryCache/Source/GeometryCache/Private/GeometryCacheSceneProxy.cpp
+++ b/Engine/Plugins/Experimental/GeometryCache/Source/GeometryCache/Private/GeometryCacheSceneProxy.cpp
@@ -606,7 +606,7 @@ void FGeometryCacheSceneProxy::UpdateAnimation(float NewTime, bool bNewLooping,
 				Section->RayTracingGeometry.Initializer.IndexBuffer = Section->IndexBuffer.IndexBufferRHI;
 				Section->RayTracingGeometry.Initializer.TotalPrimitiveCount = 0;
 
-				const uint32 IndexBufferNumTriangles = Section->IndexBuffer.NumIndices / 3;
+				const uint32 IndexBufferNumTriangles = Section->IndexBuffer.NumValidIndices / 3;
 
 				TMemoryImageArray<FRayTracingGeometrySegment>& Segments = Section->RayTracingGeometry.Initializer.Segments;
 				Segments.Reset();
@@ -619,8 +619,10 @@ void FGeometryCacheSceneProxy::UpdateAnimation(float NewTime, bool bNewLooping,
 					Segment.NumPrimitives = BatchInfo.NumTriangles;
 
 					// Ensure that a geometry segment does not access the index buffer out of bounds
-					if (!ensureMsgf(Segment.FirstPrimitive + Segment.NumPrimitives <= IndexBufferNumTriangles, 
-						TEXT("Ray tracing geometry index buffer is smaller than what's required by FGeometryCacheMeshBatchInfo")))
+					if (!ensureMsgf(Segment.FirstPrimitive + Segment.NumPrimitives <= IndexBufferNumTriangles,
+						TEXT("Ray tracing geometry index buffer is smaller than what's required by FGeometryCacheMeshBatchInfo. ")
+						TEXT("Segment.FirstPrimitive=%d Segment.NumPrimitives=%d RequiredIndexBufferTriangles=%d IndexBufferNumTriangles=%d"),
+						Segment.FirstPrimitive, Segment.NumPrimitives, Segment.FirstPrimitive+Segment.NumPrimitives, IndexBufferNumTriangles))
 					{
 						Segment.NumPrimitives = IndexBufferNumTriangles - FMath::Min<uint32>(Segment.FirstPrimitive, IndexBufferNumTriangles);
 					}
@@ -1268,7 +1270,8 @@ void FGeomCacheTrackProxy::InitRenderResources(int32 NumVertices, int32 NumIndic
 	PositionBufferFrameTimes[0] = PositionBufferFrameTimes[1] = -1.0f;
 
 	// Allocate index buffer
-	IndexBuffer.NumIndices = NumIndices;
+	IndexBuffer.NumAllocatedIndices = NumIndices;
+	IndexBuffer.NumValidIndices = 0;
 
 	// Init vertex factory
 	VertexFactory.Init(&PositionBuffers[0], &PositionBuffers[1], &TangentXBuffer, &TangentZBuffer, &TextureCoordinatesBuffer, &ColorBuffer);
@@ -1373,11 +1376,12 @@ void FGeomCacheVertexFactory::Init(const FVertexBuffer* PositionBuffer, const FV
 void FGeomCacheIndexBuffer::InitRHI()
 {
 	FRHIResourceCreateInfo CreateInfo(TEXT("FGeomCacheIndexBuffer"));
-	IndexBufferRHI = RHICreateBuffer(NumIndices * sizeof(uint32), BUF_Dynamic | BUF_IndexBuffer | BUF_ShaderResource, sizeof(uint32), ERHIAccess::VertexOrIndexBuffer | ERHIAccess::SRVMask, CreateInfo);
+	IndexBufferRHI = RHICreateBuffer(NumAllocatedIndices * sizeof(uint32), BUF_Dynamic | BUF_IndexBuffer | BUF_ShaderResource, sizeof(uint32), ERHIAccess::VertexOrIndexBuffer | ERHIAccess::SRVMask, CreateInfo);
+	NumValidIndices = 0;
 
-	if (IndexBufferRHI && NumIndices)
+	if (IndexBufferRHI && NumAllocatedIndices)
 	{
-		BufferSRV = RHICreateShaderResourceView(NumIndices ? IndexBufferRHI : nullptr);
+		BufferSRV = RHICreateShaderResourceView(NumAllocatedIndices ? IndexBufferRHI : nullptr);
 	}
 }
 
@@ -1395,32 +1399,35 @@ void FGeomCacheIndexBuffer::Update(const TArray<uint32>& Indices)
 
 	void* Buffer = nullptr;
 
+	NumValidIndices = 0;
+
 	// We only ever grow in size. Ok for now?
 	bool bReallocate = false;
-	if (Indices.Num() > NumIndices)
+	if (Indices.Num() > NumAllocatedIndices)
 	{
-		NumIndices = Indices.Num();
+		NumAllocatedIndices = Indices.Num();
 		FRHIResourceCreateInfo CreateInfo(TEXT("FGeomCacheIndexBuffer"));
-		IndexBufferRHI = RHICreateBuffer(NumIndices * sizeof(uint32), BUF_Dynamic | BUF_IndexBuffer | BUF_ShaderResource, sizeof(uint32), ERHIAccess::VertexOrIndexBuffer | ERHIAccess::SRVMask, CreateInfo);
+		IndexBufferRHI = RHICreateBuffer(NumAllocatedIndices * sizeof(uint32), BUF_Dynamic | BUF_IndexBuffer | BUF_ShaderResource, sizeof(uint32), ERHIAccess::VertexOrIndexBuffer | ERHIAccess::SRVMask, CreateInfo);
 		bReallocate = true;
 	}
 
 	if (Indices.Num() > 0)
 	{
 		// Copy the index data into the index buffer.
-		Buffer = RHILockBuffer(IndexBufferRHI, 0, NumIndices * sizeof(uint32), RLM_WriteOnly);
+		Buffer = RHILockBuffer(IndexBufferRHI, 0, NumAllocatedIndices * sizeof(uint32), RLM_WriteOnly);
 	}
 
 
 	if (Buffer)
 	{
 		FMemory::Memcpy(Buffer, Indices.GetData(), Indices.Num() * sizeof(uint32));
+		NumValidIndices = Indices.Num();
 
 		// Do not leave any of the index buffer memory uninitialized to prevent
 		// the possibility of accessing vertex buffers out of bounds.
 		uint32* LockedIndices = reinterpret_cast<uint32*>(Buffer);
 		uint32 ValidIndexValue = Indices[0];
-		for (int32 i = Indices.Num(); i < NumIndices; ++i)
+		for (int32 i = NumValidIndices; i < NumAllocatedIndices; ++i)
 		{
 			LockedIndices[i] = ValidIndexValue;
 		}
@@ -1428,9 +1435,9 @@ void FGeomCacheIndexBuffer::Update(const TArray<uint32>& Indices)
 		RHIUnlockBuffer(IndexBufferRHI);
 	}
 
-	if (bReallocate && IndexBufferRHI && NumIndices)
+	if (bReallocate && IndexBufferRHI && NumAllocatedIndices)
 	{
-		BufferSRV = RHICreateShaderResourceView(NumIndices ? IndexBufferRHI : nullptr);
+		BufferSRV = RHICreateShaderResourceView(NumAllocatedIndices ? IndexBufferRHI : nullptr);
 	}
 }
 
@@ -1440,17 +1447,18 @@ void FGeomCacheIndexBuffer::UpdateSizeOnly(int32 NewNumIndices)
 
 	// We only ever grow in size. Ok for now?
 	bool bReallocate = false;
-	if (NewNumIndices > NumIndices)
+	if (NewNumIndices > NumAllocatedIndices)
 	{
 		FRHIResourceCreateInfo CreateInfo(TEXT("FGeomCacheIndexBuffer"));
 		IndexBufferRHI = RHICreateBuffer(NewNumIndices * sizeof(uint32), BUF_Dynamic | BUF_IndexBuffer | BUF_ShaderResource, sizeof(uint32), ERHIAccess::VertexOrIndexBuffer | ERHIAccess::SRVMask, CreateInfo);
-		NumIndices = NewNumIndices;
+		NumAllocatedIndices = NewNumIndices;
+		NumValidIndices = 0;
 		bReallocate = true;
 	}
 
-	if (bReallocate && IndexBufferRHI && NumIndices)
+	if (bReallocate && IndexBufferRHI && NumAllocatedIndices)
 	{
-		BufferSRV = RHICreateShaderResourceView(NumIndices ? IndexBufferRHI : nullptr);
+		BufferSRV = RHICreateShaderResourceView(NumAllocatedIndices ? IndexBufferRHI : nullptr);
 	}
 }
 
diff --git a/Engine/Plugins/Experimental/GeometryCache/Source/GeometryCache/Public/GeometryCacheSceneProxy.h b/Engine/Plugins/Experimental/GeometryCache/Source/GeometryCache/Public/GeometryCacheSceneProxy.h
--- a/Engine/Plugins/Experimental/GeometryCache/Source/GeometryCache/Public/GeometryCacheSceneProxy.h
+++ b/Engine/Plugins/Experimental/GeometryCache/Source/GeometryCache/Public/GeometryCacheSceneProxy.h
@@ -110,7 +110,8 @@ public:
 class GEOMETRYCACHE_API FGeomCacheIndexBuffer : public FIndexBuffer
 {
 public:
-	int32 NumIndices;
+	int32 NumAllocatedIndices = 0; // Total allocated GPU index buffer size in elements
+	int32 NumValidIndices = 0; // Current valid data region of the index buffer (may be smaller than allocated buffer)
 
 	/* Create on rhi thread. Uninitialized with NumIndices space.*/
 	virtual void InitRHI() override;
@@ -124,7 +125,7 @@ public:
 
 	void UpdateSizeOnly(int32 NewNumIndices);
 
-	unsigned SizeInBytes() { return NumIndices * sizeof(uint32); }
+	unsigned SizeInBytes() { return NumAllocatedIndices * sizeof(uint32); }
 
 	FRHIShaderResourceView* GetBufferSRV() const { return BufferSRV; }
 
diff --git a/Engine/Plugins/Runtime/HairStrands/Source/HairStrandsCore/Private/GroomGeometryCache.cpp b/Engine/Plugins/Runtime/HairStrands/Source/HairStrandsCore/Private/GroomGeometryCache.cpp
--- a/Engine/Plugins/Runtime/HairStrands/Source/HairStrandsCore/Private/GroomGeometryCache.cpp
+++ b/Engine/Plugins/Runtime/HairStrands/Source/HairStrandsCore/Private/GroomGeometryCache.cpp
@@ -156,7 +156,7 @@ static void BuildBoneMatrices(USkeletalMeshComponent* SkeletalMeshComponent, con
 				{
 					CachedSection.IndexBuffer = RHICreateShaderResourceView(TrackProxy->IndexBuffer.IndexBufferRHI);
 				}
-				CachedSection.TotalIndexCount = TrackProxy->IndexBuffer.NumIndices;
+				CachedSection.TotalIndexCount = TrackProxy->IndexBuffer.NumValidIndices;
 				CachedSection.UVsChannelCount = 1;
 				CachedSection.NumPrimitives = BatchInfo.NumTriangles;
 				CachedSection.NumVertices = TrackProxy->MeshData->Positions.Num();
