- Use RHIBuildAccelerationStructures for updating the GeometryCache instead of full RHIGeometryObject recreation to reduce memory footprint (build can happen on the same allocated memory if nothing changes which requires a full rebuild)
- Add option to disable compaction on static BLAS if it will only be used for a single frame

diff --git a/Engine/Plugins/Experimental/GeometryCache/Source/GeometryCache/Private/GeometryCacheSceneProxy.cpp b/Engine/Plugins/Experimental/GeometryCache/Source/GeometryCache/Private/GeometryCacheSceneProxy.cpp
--- a/Engine/Plugins/Experimental/GeometryCache/Source/GeometryCache/Private/GeometryCacheSceneProxy.cpp
+++ b/Engine/Plugins/Experimental/GeometryCache/Source/GeometryCache/Private/GeometryCacheSceneProxy.cpp
@@ -168,6 +168,7 @@ FGeometryCacheSceneProxy::FGeometryCacheSceneProxy(UGeometryCacheComponent* Comp
 						Initializer.TotalPrimitiveCount = 0;
 						Initializer.GeometryType = RTGT_Triangles;
 						Initializer.bFastBuild = false;
+						Initializer.bAllowCompaction = false;
 
 						TArray<FRayTracingGeometrySegment> Segments;
 						const FGeometryCacheMeshData* MeshData = Section->bNextFrameMeshDataSelected ? Section->NextFrameMeshData : Section->MeshData;
@@ -602,16 +603,34 @@ void FGeometryCacheSceneProxy::UpdateAnimation(float NewTime, bool bNewLooping,
 			if (Section != nullptr)
 			{
 				const int PositionBufferIndex = Section->CurrentPositionBufferIndex != -1 ? Section->CurrentPositionBufferIndex % 2 : 0;
-
-				Section->RayTracingGeometry.Initializer.IndexBuffer = Section->IndexBuffer.IndexBufferRHI;
-				Section->RayTracingGeometry.Initializer.TotalPrimitiveCount = 0;
-
+				const FGeometryCacheMeshData* MeshData = Section->bNextFrameMeshDataSelected ? Section->NextFrameMeshData : Section->MeshData;
 				const uint32 IndexBufferNumTriangles = Section->IndexBuffer.NumValidIndices / 3;
 
 				TMemoryImageArray<FRayTracingGeometrySegment>& Segments = Section->RayTracingGeometry.Initializer.Segments;
-				Segments.Reset();
 
-				const FGeometryCacheMeshData* MeshData = Section->bNextFrameMeshDataSelected ? Section->NextFrameMeshData : Section->MeshData;
+				// Check if a full RaytracingGeometry object needs to be recreated. 
+				// Recreate when:
+				// - index buffer changes (grew in size)
+				// - total primitive count changes
+				// - segment count or vertex count changed (change BLAS size)
+				bool bRequireRecreate = false;
+				bRequireRecreate = bRequireRecreate || Segments.Num() != MeshData->BatchesInfo.Num();				
+
+				// Validate the max vertex count on all segments 
+				if (!bRequireRecreate)
+				{
+					for (int32 SegmentIndex = 0; SegmentIndex < Segments.Num(); ++SegmentIndex)
+					{
+						const FGeometryCacheMeshBatchInfo& BatchInfo = MeshData->BatchesInfo[SegmentIndex];
+						const FRayTracingGeometrySegment& Segment = Segments[SegmentIndex];
+
+						int32 MaxSegmentVertices = Section->PositionBuffers[PositionBufferIndex].GetSizeInBytes() / Segment.VertexBufferStride; // conservative estimate
+						bRequireRecreate = bRequireRecreate || Segment.MaxVertices != MaxSegmentVertices;
+					}
+				}
+
+				uint32 TotalPrimitiveCount = 0;				
+				Segments.Reset();
 				for (const FGeometryCacheMeshBatchInfo& BatchInfo : MeshData->BatchesInfo)
 				{
 					FRayTracingGeometrySegment Segment;
@@ -631,12 +650,30 @@ void FGeometryCacheSceneProxy::UpdateAnimation(float NewTime, bool bNewLooping,
 					Segment.MaxVertices = Section->PositionBuffers[PositionBufferIndex].GetSizeInBytes() / Segment.VertexBufferStride; // conservative estimate
 
 					Segments.Add(Segment);
-					Section->RayTracingGeometry.Initializer.TotalPrimitiveCount += Segment.NumPrimitives;
+					TotalPrimitiveCount += BatchInfo.NumTriangles;
 				}
 
+				bRequireRecreate = bRequireRecreate || Section->RayTracingGeometry.Initializer.IndexBuffer != Section->IndexBuffer.IndexBufferRHI;
+				bRequireRecreate = bRequireRecreate || Section->RayTracingGeometry.Initializer.TotalPrimitiveCount != TotalPrimitiveCount;
+
+				Section->RayTracingGeometry.Initializer.IndexBuffer = Section->IndexBuffer.IndexBufferRHI;
+				Section->RayTracingGeometry.Initializer.TotalPrimitiveCount = TotalPrimitiveCount;
+
 				if (Segments.Num() > 0)
 				{
-					Section->RayTracingGeometry.UpdateRHI();
+					if (bRequireRecreate)
+					{
+						Section->RayTracingGeometry.UpdateRHI();
+					}
+					else
+					{
+						// Request full build on same geometry because data might have changed to much for update call?
+						FRayTracingGeometryBuildParams BuildParams;
+						BuildParams.Geometry = Section->RayTracingGeometry.RayTracingGeometryRHI;
+						BuildParams.BuildMode = EAccelerationStructureBuildMode::Build;
+						BuildParams.Segments = Section->RayTracingGeometry.Initializer.Segments;
+						FRHICommandListExecutor::GetImmediateCommandList().BuildAccelerationStructures(MakeArrayView(&BuildParams, 1));
+					}
 				}
 			}
 		}
diff --git a/Engine/Source/Runtime/D3D12RHI/Private/D3D12RayTracing.cpp b/Engine/Source/Runtime/D3D12RHI/Private/D3D12RayTracing.cpp
--- a/Engine/Source/Runtime/D3D12RHI/Private/D3D12RayTracing.cpp
+++ b/Engine/Source/Runtime/D3D12RHI/Private/D3D12RayTracing.cpp
@@ -2910,7 +2910,7 @@ static ERayTracingAccelerationStructureFlags GetRayTracingAccelerationStructureB
 		EnumAddFlags(BuildFlags, ERayTracingAccelerationStructureFlags::AllowUpdate);
 	}
 
-	if (!Initializer.bFastBuild && !Initializer.bAllowUpdate && GD3D12RayTracingAllowCompaction)
+	if (!Initializer.bFastBuild && !Initializer.bAllowUpdate && Initializer.bAllowCompaction && GD3D12RayTracingAllowCompaction)
 	{
 		EnumAddFlags(BuildFlags, ERayTracingAccelerationStructureFlags::AllowCompaction);
 	}
diff --git a/Engine/Source/Runtime/RHI/Public/DynamicRHI.h b/Engine/Source/Runtime/RHI/Public/DynamicRHI.h
--- a/Engine/Source/Runtime/RHI/Public/DynamicRHI.h
+++ b/Engine/Source/Runtime/RHI/Public/DynamicRHI.h
@@ -206,6 +206,7 @@ public:
 
 	LAYOUT_FIELD_INITIALIZED(bool, bFastBuild, false);
 	LAYOUT_FIELD_INITIALIZED(bool, bAllowUpdate, false);
+	LAYOUT_FIELD_INITIALIZED(bool, bAllowCompaction, true);
 
 	LAYOUT_FIELD(FName, DebugName);
 };
