Conservative fix for thread safety issue in FGeometryCacheStreamingManager. Access to StreamingGeometryCaches is now protected behind FRWLock to prevent modifications from Async Loading Thread while it is being used by GameThread or workers.

diff --git a/Engine/Plugins/Experimental/GeometryCache/Source/GeometryCache/Private/GeometryCacheStreamingManager.cpp b/Engine/Plugins/Experimental/GeometryCache/Source/GeometryCache/Private/GeometryCacheStreamingManager.cpp
--- a/Engine/Plugins/Experimental/GeometryCache/Source/GeometryCache/Private/GeometryCacheStreamingManager.cpp
+++ b/Engine/Plugins/Experimental/GeometryCache/Source/GeometryCache/Private/GeometryCacheStreamingManager.cpp
@@ -74,7 +74,7 @@ private:
 	/** Scene components currently running geometry streaming. */
 	TArray<UGeometryCacheComponent*>	StreamingComponents;
 
-	mutable FCriticalSection CriticalSection;
+	mutable FRWLock StreamingGeometryCachesLock;
 
 	double LastTickTime;
 };
@@ -92,6 +92,8 @@ void FGeometryCacheStreamingManager::UpdateResourceStreaming(float DeltaTime, bo
 	SCOPE_CYCLE_COUNTER(STAT_UpdateResourceStreaming)
 	check(IsInGameThread());
 
+	FRWScopeLock ReadLock(StreamingGeometryCachesLock, SLT_ReadOnly);
+	
 	//Phase zero: Clear ChunksNeeded
 	for (auto Iter = StreamingGeometryCaches.CreateIterator(); Iter; ++Iter)
 	{
@@ -159,6 +161,7 @@ int32 FGeometryCacheStreamingManager::BlockTillAllRequestsFinished(float TimeLim
 
 	int32 Result = 0;
 
+	FRWScopeLock ReadLock(StreamingGeometryCachesLock, SLT_ReadOnly);
 	if (TimeLimit == 0.0f)
 	{
 		for (auto Iter = StreamingGeometryCaches.CreateIterator(); Iter; ++Iter)
@@ -214,6 +217,8 @@ void FGeometryCacheStreamingManager::NotifyLevelOffset(class ULevel* Level, cons
 void FGeometryCacheStreamingManager::AddGeometryCache(UGeometryCacheTrackStreamable* Cache)
 {
 	check(IsInGameThread() || IsInAsyncLoadingThread());
+
+	FRWScopeLock WriteLock(StreamingGeometryCachesLock, SLT_Write);
 	FStreamingGeometryCacheData* &CacheData = StreamingGeometryCaches.FindOrAdd(Cache);
 	if (CacheData == nullptr)
 	{
@@ -224,6 +229,8 @@ void FGeometryCacheStreamingManager::AddGeometryCache(UGeometryCacheTrackStreama
 void FGeometryCacheStreamingManager::RemoveGeometryCache(UGeometryCacheTrackStreamable* Cache)
 {
 	check(IsInGameThread());
+
+	FRWScopeLock WriteLock(StreamingGeometryCachesLock, SLT_Write);
 	FStreamingGeometryCacheData** CacheData = StreamingGeometryCaches.Find(Cache);
 	if (CacheData != nullptr)
 	{
@@ -236,12 +243,16 @@ void FGeometryCacheStreamingManager::RemoveGeometryCache(UGeometryCacheTrackStre
 bool FGeometryCacheStreamingManager::IsManagedGeometryCache(const UGeometryCacheTrackStreamable* Cache) const
 {
 	check(IsInGameThread());
+
+	FRWScopeLock ReadLock(StreamingGeometryCachesLock, SLT_ReadOnly);
 	return StreamingGeometryCaches.Contains(Cache);
 }
 
 bool FGeometryCacheStreamingManager::IsStreamingInProgress(const UGeometryCacheTrackStreamable* Cache)
 {
 	check(IsInGameThread());
+
+	FRWScopeLock ReadLock(StreamingGeometryCachesLock, SLT_ReadOnly);
 	FStreamingGeometryCacheData** CacheData = StreamingGeometryCaches.Find(Cache);
 	if (CacheData != nullptr)
 	{
@@ -269,6 +280,8 @@ void FGeometryCacheStreamingManager::PrefetchDataInternal(UGeometryCacheComponen
 {
 	check(IsInGameThread());
 	check(IsManagedComponent(CacheComponent));
+
+	FRWScopeLock ReadLock(StreamingGeometryCachesLock, SLT_ReadOnly);
 	if (CacheComponent->GeometryCache)
 	{
 		for (UGeometryCacheTrack* Track : CacheComponent->GeometryCache->Tracks)
@@ -313,6 +326,7 @@ bool FGeometryCacheStreamingManager::IsManagedComponent(const UGeometryCacheComp
 
 const uint8* FGeometryCacheStreamingManager::MapChunk(const UGeometryCacheTrackStreamable* Track, uint32 ChunkIndex, uint32* OutChunkSize)
 {
+	FRWScopeLock ReadLock(StreamingGeometryCachesLock, SLT_ReadOnly);
 	FStreamingGeometryCacheData** data = StreamingGeometryCaches.Find(Track);
 	if (data)
 	{
@@ -325,6 +339,7 @@ const uint8* FGeometryCacheStreamingManager::MapChunk(const UGeometryCacheTrackS
 
 void FGeometryCacheStreamingManager::UnmapChunk(const UGeometryCacheTrackStreamable* Track, uint32 ChunkIndex)
 {
+	FRWScopeLock ReadLock(StreamingGeometryCachesLock, SLT_ReadOnly);
 	FStreamingGeometryCacheData** data = StreamingGeometryCaches.Find(Track);
 	if (data)
 	{
