Fixing a bad merge in GeometryCacheSceneProxy.

diff --git a/Engine/Plugins/Experimental/GeometryCache/Source/GeometryCache/Private/GeometryCacheSceneProxy.cpp b/Engine/Plugins/Experimental/GeometryCache/Source/GeometryCache/Private/GeometryCacheSceneProxy.cpp
--- a/Engine/Plugins/Experimental/GeometryCache/Source/GeometryCache/Private/GeometryCacheSceneProxy.cpp
+++ b/Engine/Plugins/Experimental/GeometryCache/Source/GeometryCache/Private/GeometryCacheSceneProxy.cpp
@@ -650,7 +650,7 @@ void FGeometryCacheSceneProxy::UpdateAnimation(float NewTime, bool bNewLooping,
 					Segment.MaxVertices = Section->PositionBuffers[PositionBufferIndex].GetSizeInBytes() / Segment.VertexBufferStride; // conservative estimate
 
 					Segments.Add(Segment);
-					Section->RayTracingGeometry.Initializer.TotalPrimitiveCount += BatchInfo.NumTriangles;
+					TotalPrimitiveCount += BatchInfo.NumTriangles;
 				}
 
 				bRequireRecreate = bRequireRecreate || Section->RayTracingGeometry.Initializer.IndexBuffer != Section->IndexBuffer.IndexBufferRHI;
