Fixed GeometryCache Component "Motion Vector Scale" not being applied properly in interpolation code path.

diff --git a/Engine/Plugins/Experimental/GeometryCache/Source/GeometryCache/Private/GeometryCacheSceneProxy.cpp b/Engine/Plugins/Experimental/GeometryCache/Source/GeometryCache/Private/GeometryCacheSceneProxy.cpp
--- a/Engine/Plugins/Experimental/GeometryCache/Source/GeometryCache/Private/GeometryCacheSceneProxy.cpp
+++ b/Engine/Plugins/Experimental/GeometryCache/Source/GeometryCache/Private/GeometryCacheSceneProxy.cpp
@@ -1012,6 +1012,8 @@ void FGeometryCacheSceneProxy::FrameUpdate() const
 
 					// Unroll 4 times so we can do 4 wide SIMD
 					{
+						const VectorRegister4Float ScaledWeightA = VectorSetFloat1(OneMinusInterp * MotionVectorScale);
+						const VectorRegister4Float ScaledWeightB = VectorSetFloat1(InterpolationFactor * MotionVectorScale);
 						const FVector4f* MotionVectorsAPtr4 = (const FVector4f*)MotionVectorsAPtr;
 						const FVector4f* MotionVectorsBPtr4 = (const FVector4f*)MotionVectorsBPtr;
 						FVector4f* InterpolatedMotionVectorsPtr4 = (FVector4f*)InterpolatedMotionVectorsPtr;
@@ -1019,9 +1021,9 @@ void FGeometryCacheSceneProxy::FrameUpdate() const
 						int32 Index = 0;
 						for (; Index + 3 < NumVerts; Index += 4)
 						{
-							VectorRegister4Float MotionVector0xyz_MotionVector1x = VectorMultiplyAdd(VectorLoad(MotionVectorsAPtr4 + 0), WeightA, VectorMultiply(VectorLoad(MotionVectorsBPtr4 + 0), WeightB));
-							VectorRegister4Float MotionVector1yz_MotionVector2xy = VectorMultiplyAdd(VectorLoad(MotionVectorsAPtr4 + 1), WeightA, VectorMultiply(VectorLoad(MotionVectorsBPtr4 + 1), WeightB));
-							VectorRegister4Float MotionVector2z_MotionVector3xyz = VectorMultiplyAdd(VectorLoad(MotionVectorsAPtr4 + 2), WeightA, VectorMultiply(VectorLoad(MotionVectorsBPtr4 + 2), WeightB));
+							VectorRegister4Float MotionVector0xyz_MotionVector1x = VectorMultiplyAdd(VectorLoad(MotionVectorsAPtr4 + 0), ScaledWeightA, VectorMultiply(VectorLoad(MotionVectorsBPtr4 + 0), ScaledWeightB));
+							VectorRegister4Float MotionVector1yz_MotionVector2xy = VectorMultiplyAdd(VectorLoad(MotionVectorsAPtr4 + 1), ScaledWeightA, VectorMultiply(VectorLoad(MotionVectorsBPtr4 + 1), ScaledWeightB));
+							VectorRegister4Float MotionVector2z_MotionVector3xyz = VectorMultiplyAdd(VectorLoad(MotionVectorsAPtr4 + 2), ScaledWeightA, VectorMultiply(VectorLoad(MotionVectorsBPtr4 + 2), ScaledWeightB));
 							VectorStore(MotionVector0xyz_MotionVector1x, InterpolatedMotionVectorsPtr4 + 0);
 							VectorStore(MotionVector1yz_MotionVector2xy, InterpolatedMotionVectorsPtr4 + 1);
 							VectorStore(MotionVector2z_MotionVector3xyz, InterpolatedMotionVectorsPtr4 + 2);
