Fixed GeometryCache stream was off by one for the end frame.

diff --git a/Engine/Plugins/Experimental/GeometryCache/Source/GeometryCacheStreamer/Private/GeometryCacheStreamBase.cpp b/Engine/Plugins/Experimental/GeometryCache/Source/GeometryCacheStreamer/Private/GeometryCacheStreamBase.cpp
--- a/Engine/Plugins/Experimental/GeometryCache/Source/GeometryCacheStreamer/Private/GeometryCacheStreamBase.cpp
+++ b/Engine/Plugins/Experimental/GeometryCache/Source/GeometryCacheStreamer/Private/GeometryCacheStreamBase.cpp
@@ -301,7 +301,7 @@ void FGeometryCacheStreamBase::UpdateFramesNeeded(int32 StartFrameIndex, int32 N
 	}
 
 	// Populate the list of frame indices from given StartFrameIndex up to NumFrames or EndIndex
-	for (int32 Index = StartFrameIndex; NumFrames > 0 && Index < EndIndex; ++Index)
+	for (int32 Index = StartFrameIndex; NumFrames > 0 && Index < EndIndex + 1; ++Index)
 	{
 		AddFrameIndex(Index);
 	}
