GeometryCache:
* Added support for per-section hit proxy.

diff --git a/Engine/Plugins/Experimental/GeometryCache/Source/GeometryCache/Private/GeometryCacheSceneProxy.cpp b/Engine/Plugins/Experimental/GeometryCache/Source/GeometryCache/Private/GeometryCacheSceneProxy.cpp
--- a/Engine/Plugins/Experimental/GeometryCache/Source/GeometryCache/Private/GeometryCacheSceneProxy.cpp
+++ b/Engine/Plugins/Experimental/GeometryCache/Source/GeometryCache/Private/GeometryCacheSceneProxy.cpp
@@ -2,19 +2,21 @@
 
 #include "GeometryCacheSceneProxy.h"
 #include "Async/ParallelFor.h"
-#include "MaterialShared.h"
-#include "SceneManagement.h"
-#include "EngineGlobals.h"
-#include "Materials/Material.h"
+#include "Components/BrushComponent.h"
 #include "Engine/Engine.h"
+#include "EngineGlobals.h"
+#include "EngineUtils.h"
 #include "GeometryCacheComponent.h"
 #include "GeometryCacheMeshData.h"
 #include "GeometryCache.h"
 #include "GeometryCacheTrackStreamable.h"
 #include "GeometryCacheModule.h"
 #include "GeometryCacheHelpers.h"
+#include "Materials/Material.h"
+#include "MaterialShared.h"
 #include "RayTracingDefinitions.h"
 #include "RayTracingInstance.h"
+#include "SceneManagement.h"
 
 DECLARE_CYCLE_STAT(TEXT("Gather Mesh Elements"), STAT_GeometryCacheSceneProxy_GetMeshElements, STATGROUP_GeometryCache);
 DECLARE_DWORD_COUNTER_STAT(TEXT("Triangle Count"), STAT_GeometryCacheSceneProxy_TriangleCount, STATGROUP_GeometryCache);
@@ -404,6 +406,51 @@ void FGeometryCacheSceneProxy::CreateMeshBatch(
 	Mesh.bCanApplyViewModeOverrides = false;
 }
 
+#if WITH_EDITOR
+HHitProxy* FGeometryCacheSceneProxy::CreateHitProxies(UPrimitiveComponent* Component, TArray<TRefCountPtr<HHitProxy> >& OutHitProxies)
+{
+	// Add a default hit proxy to handle cases where the number of batches changes during the animation,
+	// including when the initial frame has no mesh data
+	HHitProxy* DefaultHitProxy = FPrimitiveSceneProxy::CreateHitProxies(Component, OutHitProxies);
+
+	if (Component->GetOwner() && Tracks.Num() > 0)
+	{
+		int32 SectionIndex = 0;
+		for (const FGeomCacheTrackProxy* Track : Tracks)
+		{
+			if (Track->MeshData)
+			{
+				for (const FGeometryCacheMeshBatchInfo& BatchInfo : Track->MeshData->BatchesInfo)
+				{
+					HHitProxy* ActorHitProxy;
+
+					int32 MaterialIndex = BatchInfo.MaterialIndex;
+					if (Component->GetOwner()->IsA(ABrush::StaticClass()) && Component->IsA(UBrushComponent::StaticClass()))
+					{
+						ActorHitProxy = new HActor(Component->GetOwner(), Component, HPP_Wireframe, SectionIndex, MaterialIndex);
+					}
+					else
+					{
+						ActorHitProxy = new HActor(Component->GetOwner(), Component, Component->HitProxyPriority, SectionIndex, MaterialIndex);
+					}
+
+					OutHitProxies.Add(ActorHitProxy);
+					++SectionIndex;
+				}
+			}
+		}
+	}
+
+	HitProxyIds.SetNumUninitialized(OutHitProxies.Num());
+	for (int32 Index = 0; Index < HitProxyIds.Num(); ++Index)
+	{
+		HitProxyIds[Index] = OutHitProxies[Index]->Id;
+	}
+
+	return DefaultHitProxy;
+}
+#endif
+
 void FGeometryCacheSceneProxy::GetDynamicMeshElements(const TArray<const FSceneView*>& Views, const FSceneViewFamily& ViewFamily, uint32 VisibilityMap, FMeshElementCollector& Collector) const
 {
 	SCOPE_CYCLE_COUNTER(STAT_GeometryCacheSceneProxy_GetMeshElements);
@@ -480,6 +527,14 @@ void FGeometryCacheSceneProxy::GetDynamicMeshElements(const TArray<const FSceneV
 						FDynamicPrimitiveUniformBuffer& DynamicPrimitiveUniformBuffer = Collector.AllocateOneFrameResource<FDynamicPrimitiveUniformBuffer>();
 						CreateMeshBatch(TrackProxy, BatchInfo, UserDataWrapper, DynamicPrimitiveUniformBuffer, MeshBatch);
 
+#if WITH_EDITOR
+						// It's possible the number of batches has changed since the initial frame so validate the BatchIndex
+						if (HitProxyIds.IsValidIndex(BatchIndex))
+						{
+							MeshBatch.BatchHitProxyId = HitProxyIds[BatchIndex];
+						}
+#endif
+
 						// Apply view mode material overrides
 						FMaterialRenderProxy* MaterialProxy = bWireframe ? WireframeMaterialInstance : TrackProxy->Materials[BatchIndex]->GetRenderProxy();
 						MeshBatch.bWireframe = bWireframe;
diff --git a/Engine/Plugins/Experimental/GeometryCache/Source/GeometryCache/Public/GeometryCacheSceneProxy.h b/Engine/Plugins/Experimental/GeometryCache/Source/GeometryCache/Public/GeometryCacheSceneProxy.h
--- a/Engine/Plugins/Experimental/GeometryCache/Source/GeometryCache/Public/GeometryCacheSceneProxy.h
+++ b/Engine/Plugins/Experimental/GeometryCache/Source/GeometryCache/Public/GeometryCacheSceneProxy.h
@@ -2,6 +2,7 @@
 #pragma once
 
 #include "CoreMinimal.h"
+#include "HitProxies.h"
 #include "Stats/Stats.h"
 #include "Containers/ResourceArray.h"
 #include "RenderResource.h"
@@ -289,6 +290,9 @@ public:
 	virtual ~FGeometryCacheSceneProxy();
 
 	// Begin FPrimitiveSceneProxy interface.
+#if WITH_EDITOR
+	virtual HHitProxy* CreateHitProxies(UPrimitiveComponent* Component, TArray<TRefCountPtr<HHitProxy> >& OutHitProxies) override;
+#endif
 	virtual void GetDynamicMeshElements(const TArray<const FSceneView*>& Views, const FSceneViewFamily& ViewFamily, uint32 VisibilityMap, FMeshElementCollector& Collector) const override;
 	virtual FPrimitiveViewRelevance GetViewRelevance(const FSceneView* View) const override;
 
@@ -332,6 +336,10 @@ private:
 	/** Array of Track Proxies */
 	TArray<FGeomCacheTrackProxy*> Tracks;
 
+#if WITH_EDITOR
+	TArray<FHitProxyId> HitProxyIds;
+#endif
+
 	/** Scratch memory for frame update - do not use directly. */
 	struct FScratchMemory
 	{
