GeometryCache:
* Fixed GeometryCache not outputting velocity in base pass for motion blur. Setting r.BasePassForceOutputsVelocity=1 is not required anymore to get motion blur with MRQ.

diff --git a/Engine/Plugins/Experimental/GeometryCache/Source/GeometryCache/Private/GeometryCacheSceneProxy.cpp b/Engine/Plugins/Experimental/GeometryCache/Source/GeometryCache/Private/GeometryCacheSceneProxy.cpp
--- a/Engine/Plugins/Experimental/GeometryCache/Source/GeometryCache/Private/GeometryCacheSceneProxy.cpp
+++ b/Engine/Plugins/Experimental/GeometryCache/Source/GeometryCache/Private/GeometryCacheSceneProxy.cpp
@@ -390,7 +390,13 @@ void FGeometryCacheSceneProxy::CreateMeshBatch(
 
 	const FMatrix& LocalToWorldTransform = TrackProxy->WorldMatrix * GetLocalToWorld();
 
-	DynamicPrimitiveUniformBuffer.Set(LocalToWorldTransform, LocalToWorldTransform, GetBounds(), GetLocalBounds(), true, false, DrawsVelocity(), false);
+	bool bHasPrecomputedVolumetricLightmap;
+	FMatrix PreviousLocalToWorld;
+	int32 SingleCaptureIndex;
+	bool bOutputVelocity;
+	GetScene().GetPrimitiveUniformShaderParameters_RenderThread(GetPrimitiveSceneInfo(), bHasPrecomputedVolumetricLightmap, PreviousLocalToWorld, SingleCaptureIndex, bOutputVelocity);
+
+	DynamicPrimitiveUniformBuffer.Set(LocalToWorldTransform, PreviousLocalToWorld, GetBounds(), GetLocalBounds(), true, false, DrawsVelocity(), bOutputVelocity);
 	BatchElement.PrimitiveUniformBufferResource = &DynamicPrimitiveUniformBuffer.UniformBuffer;
 
 	const FGeometryCacheMeshData* MeshData = TrackProxy->bNextFrameMeshDataSelected ? TrackProxy->NextFrameMeshData : TrackProxy->MeshData;
