* Moved AlembicImporter and GeometryCache plugins out of experimental.

diff --git a/Engine/Plugins/Importers/AlembicImporter/Source/AlembicImporter/Classes/AlembicImportFactory.h b/Engine/Plugins/Importers/AlembicImporter/Source/AlembicImporter/Classes/AlembicImportFactory.h
new file mode 100644
--- /dev/null
+++ b/Engine/Plugins/Importers/AlembicImporter/Source/AlembicImporter/Classes/AlembicImportFactory.h
@@ -0,0 +1,109 @@
+// Copyright Epic Games, Inc. All Rights Reserved.
+
+#pragma once
+
+#include "Factories/Factory.h"
+#include "EditorReimportHandler.h"
+#include "AlembicImportFactory.generated.h"
+
+class UAbcImportSettings;
+class FAbcImporter;
+class UGeometryCache;
+class UStaticMesh;
+class USkeletalMesh;
+class UAbcAssetImportData;
+class SAlembicImportOptions;
+class SAlembicTrackSelectionWindow;
+
+class FAbcImporter;
+
+UCLASS(hidecategories = Object)
+class ALEMBICIMPORTER_API UAlembicImportFactory : public UFactory, public FReimportHandler
+{
+	GENERATED_UCLASS_BODY()
+
+	/** Object used to show import options for Alembic */
+	UPROPERTY()
+	TObjectPtr<UAbcImportSettings> ImportSettings;
+
+	UPROPERTY()
+	bool bShowOption;
+
+	//~ Begin UObject Interface
+	void PostInitProperties();
+	//~ End UObject Interface
+
+	//~ Begin UFactory Interface
+	virtual FText GetDisplayName() const override;
+	virtual bool DoesSupportClass(UClass * Class) override;
+	virtual UClass* ResolveSupportedClass() override;
+	virtual bool FactoryCanImport(const FString& Filename) override;
+	virtual UObject* FactoryCreateFile(UClass* InClass, UObject* InParent, FName InName, EObjectFlags Flags, const FString& Filename, const TCHAR* Parms, FFeedbackContext* Warn, bool& bOutOperationCanceled) override;
+	//~ End UFactory Interface
+
+	//~ Begin FReimportHandler Interface
+	virtual bool CanReimport(UObject* Obj, TArray<FString>& OutFilenames) override;
+	virtual void SetReimportPaths(UObject* Obj, const TArray<FString>& NewReimportPaths) override;
+	virtual EReimportResult::Type Reimport(UObject* Obj) override;
+
+	void ShowImportOptionsWindow(TSharedPtr<SAlembicImportOptions>& Options, FString FilePath, const FAbcImporter& Importer);
+
+	virtual int32 GetPriority() const override;
+	//~ End FReimportHandler Interface
+		
+	/**
+	* Imports a StaticMesh (using AbcImporter) from the Alembic File
+	*
+	* @param Importer - AbcImporter instance
+	* @param InParent - Parent for the StaticMesh
+	* @param Flags - Creation Flags for the StaticMesh	
+	* @return UObject*
+	*/
+	TArray<UObject*> ImportStaticMesh(FAbcImporter& Importer, UObject* InParent, EObjectFlags Flags);
+	
+	/**
+	* ImportGeometryCache
+	*	
+	* @param Importer - AbcImporter instance
+	* @param InParent - Parent for the GeometryCache asset
+	* @param Flags - Creation Flags for the GeometryCache asset
+	* @return UObject*
+	*/
+	UObject* ImportGeometryCache(FAbcImporter& Importer, UObject* InParent, EObjectFlags Flags);
+
+	/**
+	* ImportGeometryCache
+	*	
+	* @param Importer - AbcImporter instance
+	* @param InParent - Parent for the GeometryCache asset
+	* @param Flags - Creation Flags for the GeometryCache asset
+	* @return UObject*
+	*/
+	TArray<UObject*> ImportSkeletalMesh(FAbcImporter& Importer, UObject* InParent, EObjectFlags Flags);
+
+	/**
+	* ReimportStaticMesh
+	*
+	* @param Mesh - Static mesh instance to re import
+	* @return EReimportResult::Type
+	*/
+	EReimportResult::Type ReimportStaticMesh(UStaticMesh* Mesh);
+
+	/**
+	* ReimportGeometryCache
+	*
+	* @param Cache - Geometry cache instance to re import
+	* @return EReimportResult::Type
+	*/
+	EReimportResult::Type ReimportGeometryCache(UGeometryCache* Cache);
+
+	/**
+	* ReimportGeometryCache
+	*
+	* @param Cache - Geometry cache instance to re import
+	* @return EReimportResult::Type
+	*/
+	EReimportResult::Type ReimportSkeletalMesh(USkeletalMesh* SkeletalMesh);
+
+	void PopulateOptionsWithImportData(UAbcAssetImportData* ImportData);
+};
\ No newline at end of file
diff --git a/Engine/Plugins/Importers/AlembicImporter/Source/AlembicLibrary/Public/AbcImporter.h b/Engine/Plugins/Importers/AlembicImporter/Source/AlembicLibrary/Public/AbcImporter.h
new file mode 100644
--- /dev/null
+++ b/Engine/Plugins/Importers/AlembicImporter/Source/AlembicLibrary/Public/AbcImporter.h
@@ -0,0 +1,222 @@
+// Copyright Epic Games, Inc. All Rights Reserved.
+
+#pragma once
+
+#include "CoreMinimal.h"
+#include "Containers/List.h"
+#include "Animation/MorphTarget.h"
+#include "Animation/AnimSequence.h"
+#include "AbcPolyMesh.h"
+
+#if PLATFORM_WINDOWS
+#include "Windows/WindowsHWrapper.h"
+#include "Windows/AllowWindowsPlatformTypes.h"
+#endif
+
+THIRD_PARTY_INCLUDES_START
+#include <Alembic/AbcGeom/All.h>
+THIRD_PARTY_INCLUDES_END
+
+#if PLATFORM_WINDOWS
+#include "Windows/HideWindowsPlatformTypes.h"
+#endif
+
+class UMaterial;
+class UStaticMesh;
+class USkeletalMesh;
+class UGeometryCache;
+class UDEPRECATED_GeometryCacheTrack_FlipbookAnimation;
+class UDEPRECATED_GeometryCacheTrack_TransformAnimation;
+class UGeometryCacheTrackStreamable;
+class UAbcImportSettings;
+class FSkeletalMeshImportData;
+class UAbcAssetImportData;
+
+struct FMeshDescription;
+struct FAbcImportData;
+struct FGeometryCacheMeshData;
+struct FAbcMeshSample;
+struct FAbcCompressionSettings;
+
+enum EAbcImportError : uint32
+{
+	AbcImportError_NoError,
+	AbcImportError_InvalidArchive,
+	AbcImportError_NoValidTopObject,	
+	AbcImportError_NoMeshes,
+	AbcImportError_FailedToImportData
+};
+
+struct FCompressedAbcData
+{
+	~FCompressedAbcData();
+	/** GUID identifying the poly mesh object this compressed data corresponds to */
+	FGuid Guid;
+	/** Average sample to apply the bases to */
+	FAbcMeshSample* AverageSample;
+	/** List of base samples calculated using PCA compression */
+	TArray<FAbcMeshSample*> BaseSamples;
+	/** Contains the curve values for each individual base */
+	TArray<TArray<float>> CurveValues;
+	/** Contains the time key values for each individual base */
+	TArray<TArray<float>> TimeValues;
+	/** Material names used for retrieving created materials */
+	TArray<FString> MaterialNames;
+};
+
+/** Mesh section used for chunking the mesh data during Skeletal mesh building */
+struct FMeshSection
+{
+	FMeshSection() : MaterialIndex(0), NumFaces(0) {}
+	int32 MaterialIndex;
+	TArray<uint32> Indices;
+	TArray<uint32> OriginalIndices;
+	TArray<FVector> TangentX;
+	TArray<FVector> TangentY;
+	TArray<FVector> TangentZ;
+	TArray<FVector2D> UVs[MAX_TEXCOORDS];
+	TArray<FColor> Colors;
+	uint32 NumFaces;
+	uint32 NumUVSets;
+};
+
+class ALEMBICLIBRARY_API FAbcImporter
+{
+public:
+	FAbcImporter();
+	~FAbcImporter();
+public:
+	/**
+	* Opens and caches basic data from the Alembic file to be used for populating the importer UI
+	* 
+	* @param InFilePath - Path to the Alembic file to be opened
+	* @return - Possible error code 
+	*/
+	const EAbcImportError OpenAbcFileForImport(const FString InFilePath);
+
+	/**
+	* Imports the individual tracks from the Alembic file
+	*
+	* @param InNumThreads - Number of threads to use for importing
+	* @return - Possible error code
+	*/
+	const EAbcImportError ImportTrackData(const int32 InNumThreads, UAbcImportSettings* ImportSettings);
+	
+	/**
+	* Import Alembic meshes as a StaticMeshInstance
+	*	
+	* @param InParent - ParentObject for the static mesh
+	* @param Flags - Object flags
+	* @return FStaticMesh*
+	*/
+	const TArray<UStaticMesh*> ImportAsStaticMesh(UObject* InParent, EObjectFlags Flags);
+
+	/**
+	* Import an Alembic file as a GeometryCache
+	*	
+	* @param InParent - ParentObject for the static mesh
+	* @param Flags - Object flags
+	* @return UGeometryCache*
+	*/
+	UGeometryCache* ImportAsGeometryCache(UObject* InParent, EObjectFlags Flags);
+
+	TArray<UObject*> ImportAsSkeletalMesh(UObject* InParent, EObjectFlags Flags);	
+
+	/**
+	* Reimport an Alembic mesh
+	*
+	* @param Mesh - Current StaticMesh instance
+	* @return FStaticMesh*
+	*/
+	const TArray<UStaticMesh*> ReimportAsStaticMesh(UStaticMesh* Mesh);
+
+	/**
+	* Reimport an Alembic file as a GeometryCache
+	*
+	* @param GeometryCache - Current GeometryCache instance
+	* @return UGeometryCache*
+	*/
+	UGeometryCache* ReimportAsGeometryCache(UGeometryCache* GeometryCache);
+
+	/**
+	* Reimport an Alembic file as a SkeletalMesh
+	*
+	* @param SkeletalMesh - Current SkeletalMesh instance
+	* @return USkeletalMesh*
+	*/
+	TArray<UObject*> ReimportAsSkeletalMesh(USkeletalMesh* SkeletalMesh);
+
+	/** Returns the array of imported PolyMesh objects */
+	const TArray<class FAbcPolyMesh*>& GetPolyMeshes() const;
+	
+	/** Returns the lowest frame index containing data for the imported Alembic file */
+	const uint32 GetStartFrameIndex() const;
+
+	/** Returns the highest frame index containing data for the imported Alembic file */
+	const uint32 GetEndFrameIndex() const;
+
+	/** Returns the number of tracks found in the imported Alembic file */
+	const uint32 GetNumMeshTracks() const;
+
+	void UpdateAssetImportData(UAbcAssetImportData* AssetImportData);
+	void RetrieveAssetImportData(UAbcAssetImportData* ImportData);
+private:
+	/**
+	* Creates an template object instance taking into account existing Instances and Objects (on reimporting)
+	*
+	* @param InParent - ParentObject for the geometry cache, this can be changed when parent is deleted/re-created
+	* @param ObjectName - Name to be used for the created object
+	* @param Flags - Object creation flags
+	*/
+	template<typename T> T* CreateObjectInstance(UObject*& InParent, const FString& ObjectName, const EObjectFlags Flags);
+	
+	/**
+	* Creates a Static mesh from the given mesh description
+	*
+	* @param InParent - ParentObject for the static mesh
+	* @param Name - Name for the static mesh
+	* @param Flags - Object flags
+	* @param NumMaterials - Number of materials to add
+	* @param FaceSetNames - Face set names used for retrieving the materials
+	* @param MeshDescription - The MeshDescription from which the static mesh should be constructed
+	* @return UStaticMesh*
+	*/
+	UStaticMesh* CreateStaticMeshFromSample(UObject* InParent, const FString& Name, EObjectFlags Flags, const uint32 NumMaterials, const TArray<FString>& FaceSetNames, const FAbcMeshSample* Sample);
+
+	/** Generates and populate a FMeshDescription instance from the given sample*/
+	void GenerateMeshDescriptionFromSample(const FAbcMeshSample* Sample, FMeshDescription* MeshDescription, UStaticMesh* StaticMesh);
+	
+	/** Compresses the imported animation data, returns true if compression was successful and compressed data was populated */
+	const bool CompressAnimationDataUsingPCA(const FAbcCompressionSettings& InCompressionSettings, const bool bRunComparison = false);	
+	/** Performs the actual SVD compression to retrieve the bases and weights used to set up the Skeletal mesh's morph targets */
+	const int32 PerformSVDCompression(const TArray<float>& OriginalMatrix, const TArray<float>& OriginalNormalsMatrix, const uint32 NumSamples, const float InPercentage, const int32 InFixedNumValue,
+		TArray<float>& OutU, TArray<float>& OutNormalsU, TArray<float>& OutV);
+
+	/** Functionality for comparing the matrices and calculating the difference from the original animation */
+	void CompareCompressionResult(const TArray<float>& OriginalMatrix, const uint32 NumSamples, const uint32 NumUsedSingularValues, const TArrayView<float>& OutU, const TArray<float>& OutV, const float Tolerance);
+	
+	/** Build a skeletal mesh from the PCA compressed data */
+	bool BuildSkeletalMesh(FSkeletalMeshLODModel& LODModel, const FReferenceSkeleton& RefSkeleton, FAbcMeshSample* Sample, TArray<int32>& OutMorphTargetVertexRemapping, TArray<int32>& OutUsedVertexIndicesForMorphs);
+	
+	/** Generate morph target vertices from the PCA compressed bases */
+	void GenerateMorphTargetVertices(FAbcMeshSample* BaseSample, TArray<FMorphTargetDelta> &MorphDeltas, FAbcMeshSample* AverageSample, uint32 WedgeOffset, const TArray<int32>& RemapIndices, const TArray<int32>& UsedVertexIndicesForMorphs, const uint32 VertexOffset, const uint32 IndexOffset);
+	
+	/** Set up correct morph target weights from the PCA compressed data */
+	void SetupMorphTargetCurves(USkeleton* Skeleton, FName ConstCurveName, UAnimSequence* Sequence, const TArray<float> &CurveValues, const TArray<float>& TimeValues, IAnimationDataController& Controller);
+	
+	/** Set the Alembic archive metadata on the given objects */
+	void SetMetaData(const TArray<UObject*>& Objects);
+
+private:
+	/** Cached ptr for the import settings */
+	UAbcImportSettings* ImportSettings;
+
+	/** Resulting compressed data from PCA compression */
+	TArray<FCompressedAbcData> CompressedMeshData;
+
+	/** Offset for each sample, used as the root bone translation */
+	TOptional<TArray<FVector>> SamplesOffsets;
+
+	/** ABC file representation for currently opened filed */
+	class FAbcFile* AbcFile;
+};
diff --git a/Engine/Plugins/Importers/AlembicImporter/Source/AlembicLibrary/Public/AbcTransform.h b/Engine/Plugins/Importers/AlembicImporter/Source/AlembicLibrary/Public/AbcTransform.h
new file mode 100644
--- /dev/null
+++ b/Engine/Plugins/Importers/AlembicImporter/Source/AlembicLibrary/Public/AbcTransform.h
@@ -0,0 +1,33 @@
+// Copyright Epic Games, Inc. All Rights Reserved.
+
+#pragma once
+
+#include "AbcObject.h"
+
+class ALEMBICLIBRARY_API FAbcTransform : public IAbcObject
+{
+public:
+	FAbcTransform(const Alembic::AbcGeom::IXform& InTransform, const FAbcFile* InFile, IAbcObject* InParent = nullptr);
+	virtual ~FAbcTransform() {}
+
+	/** Begin IAbcObject overrides */
+	virtual bool ReadFirstFrame(const float InTime, const int32 FrameIndex) final;
+	virtual void SetFrameAndTime(const float InTime, const int32 FrameIndex, const EFrameReadFlags InFlags, const int32 TargetIndex = INDEX_NONE) final;
+	virtual FMatrix GetMatrix(const int32 FrameIndex) const final;
+	virtual bool HasConstantTransform() const final;
+	virtual void PurgeFrameData(const int32 FrameIndex) final;
+	/** End IAbcObject overrides */	
+public:
+	/** Flag whether or not this transformation object is identity constant */
+	bool bConstantIdentity;
+protected:
+	/** Alembic representation of this object */
+	const Alembic::AbcGeom::IXform Transform;
+	/** Schema extracted from transform object  */
+	const Alembic::AbcGeom::IXformSchema Schema;
+
+	/** Initial value for this object in first frame with available data */
+	FMatrix InitialValue;
+	/** Resident set of matrix values for this object, used for parallel reading of samples/frames */
+	FMatrix ResidentMatrices[MaxNumberOfResidentSamples];
+};
\ No newline at end of file
diff --git a/Engine/Plugins/Importers/AlembicImporter/Source/AlembicLibrary/Public/AbcUtilities.h b/Engine/Plugins/Importers/AlembicImporter/Source/AlembicLibrary/Public/AbcUtilities.h
new file mode 100644
--- /dev/null
+++ b/Engine/Plugins/Importers/AlembicImporter/Source/AlembicLibrary/Public/AbcUtilities.h
@@ -0,0 +1,20 @@
+// Copyright Epic Games, Inc. All Rights Reserved.
+
+#pragma once
+
+#include "CoreMinimal.h"
+
+class FAbcFile;
+struct FGeometryCacheMeshData;
+
+class ALEMBICLIBRARY_API FAbcUtilities
+{
+public:
+	/*  Populates a FGeometryCacheMeshData instance (with merged meshes) for the given frame of an Alembic 
+	 *  with ConcurrencyIndex used as the concurrent read index
+	 */
+	static void GetFrameMeshData(FAbcFile& AbcFile, int32 FrameIndex, FGeometryCacheMeshData& OutMeshData, int32 ConcurrencyIndex = 0);
+
+	/** Sets up materials from an AbcFile to a GeometryCache and moves them into the given package */
+	static void SetupGeometryCacheMaterials(FAbcFile& AbcFile, class UGeometryCache* GeometryCache, UObject* Package);
+};
\ No newline at end of file
diff --git a/Engine/Plugins/Runtime/GeometryCache/Source/GeometryCache/Classes/GeometryCache.h b/Engine/Plugins/Runtime/GeometryCache/Source/GeometryCache/Classes/GeometryCache.h
new file mode 100644
--- /dev/null
+++ b/Engine/Plugins/Runtime/GeometryCache/Source/GeometryCache/Classes/GeometryCache.h
@@ -0,0 +1,101 @@
+// Copyright Epic Games, Inc. All Rights Reserved.
+#pragma once
+
+#include "CoreMinimal.h"
+#include "UObject/ObjectMacros.h"
+#include "UObject/Object.h"
+#include "Interfaces/Interface_AssetUserData.h"
+#include "RenderCommandFence.h"
+
+#include "GeometryCache.generated.h"
+
+class UGeometryCacheTrack;
+class UMaterialInterface;
+struct FGeometryCacheMeshData;
+
+DECLARE_LOG_CATEGORY_EXTERN(LogGeometryCache, Log, All);
+
+/**
+* A Geometry Cache is a piece/set of geometry that consists of individual Mesh/Transformation samples.
+* In contrast with Static Meshes they can have their vertices animated in certain ways. * 
+*/
+UCLASS(hidecategories = Object, BlueprintType, config = Engine)
+class GEOMETRYCACHE_API UGeometryCache : public UObject, public IInterface_AssetUserData
+{
+	GENERATED_UCLASS_BODY()
+public:
+	//~ Begin UObject Interface.
+	virtual void Serialize(FArchive& Ar) override;
+	virtual void PostInitProperties() override;
+	virtual FString GetDesc() override;
+	virtual void GetAssetRegistryTags(TArray<FAssetRegistryTag>& OutTags) const override;
+
+	virtual void BeginDestroy() override;
+	virtual bool IsReadyForFinishDestroy() override;
+#if WITH_EDITOR
+	virtual void PreEditChange(FProperty* PropertyAboutToChange) override;
+#endif // WITH_EDITOR
+	//~ End UObject Interface.
+
+	/**
+	* AddTrack
+	*
+	* @param Track - GeometryCacheTrack instance that is a part of the GeometryCacheAsset
+	* @return void
+	*/
+	void AddTrack(UGeometryCacheTrack* Track);
+		
+	/** Clears all stored data so the reimporting step can fill the instance again*/
+	void ClearForReimporting();	
+
+#if WITH_EDITORONLY_DATA
+	/** Importing data and options used for this Geometry cache object*/
+	UPROPERTY(Category = ImportSettings, VisibleAnywhere, Instanced)
+	TObjectPtr<class UAssetImportData> AssetImportData;	
+
+	/** Information for thumbnail rendering */
+	UPROPERTY(VisibleAnywhere, Instanced, Category = Thumbnail)
+	TObjectPtr<class UThumbnailInfo> ThumbnailInfo;
+#endif
+	UPROPERTY(EditAnywhere, Category = GeometryCache)
+	TArray<TObjectPtr<UMaterialInterface>> Materials;
+	
+	/** GeometryCache track defining the samples/geometry data for this GeomCache instance */
+	UPROPERTY(VisibleAnywhere, Category=GeometryCache)
+	TArray<TObjectPtr<UGeometryCacheTrack>> Tracks;
+
+	/** Set the start and end frames for the GeometryCache */
+	void SetFrameStartEnd(int32 InStartFrame, int32 InEndFrame);
+
+	/** Get the start frame */
+	int32 GetStartFrame() const;
+
+	/** Get the end frame */
+	int32 GetEndFrame() const;
+	
+	/** Calculate it's duration */
+	float CalculateDuration() const;
+
+	/** Get the Frame at the Specified Time */
+	int32  GetFrameAtTime(const float Time) const;
+
+	/** Get the mesh data at the specified time */
+	void GetMeshDataAtTime(float Time, TArray<FGeometryCacheMeshData>& OutMeshData) const;
+
+	/** Get the hash of the meshes data of the GeometryCache */
+	FString GetHash() const;
+
+private:
+	/** A fence which is used to keep track of the rendering thread releasing the geometry cache resources. */
+	FRenderCommandFence ReleaseResourcesFence;
+
+protected:
+	UPROPERTY(BlueprintReadOnly, Category = GeometryCache)
+	int32 StartFrame;
+
+	UPROPERTY(BlueprintReadOnly, Category = GeometryCache)
+	int32 EndFrame;
+
+	UPROPERTY()
+	uint64 Hash;
+};
diff --git a/Engine/Plugins/Runtime/GeometryCache/Source/GeometryCache/Classes/GeometryCacheCodecRaw.h b/Engine/Plugins/Runtime/GeometryCache/Source/GeometryCache/Classes/GeometryCacheCodecRaw.h
new file mode 100644
--- /dev/null
+++ b/Engine/Plugins/Runtime/GeometryCache/Source/GeometryCache/Classes/GeometryCacheCodecRaw.h
@@ -0,0 +1,46 @@
+// Copyright Epic Games, Inc. All Rights Reserved.
+#pragma once
+
+#include "CoreMinimal.h"
+#include "UObject/ObjectMacros.h"
+#include "UObject/ScriptMacros.h"
+#include "GeometryCacheCodecBase.h"
+
+#include "GeometryCacheCodecRaw.generated.h"
+
+class UGeometryCacheTrackStreamable;
+
+struct FGeometryCacheCodecRenderStateRaw : FGeometryCacheCodecRenderStateBase
+{
+	FGeometryCacheCodecRenderStateRaw(const TArray<int32> &SetTopologyRanges) : FGeometryCacheCodecRenderStateBase(SetTopologyRanges) {}
+	virtual bool DecodeSingleFrame(FGeometryCacheCodecDecodeArguments &Args) override;
+};
+
+UCLASS(hidecategories = Object)
+class GEOMETRYCACHE_API UGeometryCacheCodecRaw : public UGeometryCacheCodecBase
+{
+	GENERATED_UCLASS_BODY()
+
+	virtual bool DecodeSingleFrame(FGeometryCacheCodecDecodeArguments &Args) override;
+	virtual FGeometryCacheCodecRenderStateBase *CreateRenderState() override;
+
+#if WITH_EDITORONLY_DATA
+	virtual void BeginCoding(TArray<FStreamedGeometryCacheChunk> &AppendChunksTo) override;
+	virtual void EndCoding() override;
+	virtual void CodeFrame(const FGeometryCacheCodecEncodeArguments& Args) override;
+#endif
+
+private:
+	UPROPERTY(VisibleAnywhere, Category = GeometryCache)
+	int32 DummyProperty;
+
+#if WITH_EDITORONLY_DATA
+	struct FEncoderData
+	{
+		int32 CurrentChunkId; // Index in the AppendChunksTo list of the chunk we are working on
+		TArray<FStreamedGeometryCacheChunk> *AppendChunksTo; // Any chunks this codec creates will be appended to this list
+	};
+	FEncoderData EncoderData;
+#endif
+};
+
diff --git a/Engine/Plugins/Runtime/GeometryCache/Source/GeometryCache/Classes/GeometryCacheCodecV1.h b/Engine/Plugins/Runtime/GeometryCache/Source/GeometryCache/Classes/GeometryCacheCodecV1.h
new file mode 100644
--- /dev/null
+++ b/Engine/Plugins/Runtime/GeometryCache/Source/GeometryCache/Classes/GeometryCacheCodecV1.h
@@ -0,0 +1,58 @@
+// Copyright Epic Games, Inc. All Rights Reserved.
+#pragma once
+
+#include "CoreMinimal.h"
+#include "UObject/ObjectMacros.h"
+#include "UObject/ScriptMacros.h"
+#include "GeometryCacheCodecBase.h"
+#include "Serialization/MemoryWriter.h"
+#include "Serialization/BufferReader.h"
+
+#include "GeometryCacheCodecV1.generated.h"
+
+class UGeometryCacheTrackStreamable;
+class ICodecDecoder;
+class ICodecEncoder;
+
+struct FGeometryCacheCodecRenderStateV1 : FGeometryCacheCodecRenderStateBase
+{
+	FGeometryCacheCodecRenderStateV1(const TArray<int32> &SetTopologyRanges); 
+	~FGeometryCacheCodecRenderStateV1();	
+	
+	virtual bool DecodeSingleFrame(FGeometryCacheCodecDecodeArguments &Args) override;
+};
+
+UCLASS(hidecategories = Object)
+class GEOMETRYCACHE_API UGeometryCacheCodecV1 : public UGeometryCacheCodecBase
+{
+	GENERATED_UCLASS_BODY()
+
+	virtual ~UGeometryCacheCodecV1();
+
+	virtual bool DecodeSingleFrame(FGeometryCacheCodecDecodeArguments &Args) override;
+	virtual bool DecodeBuffer(const uint8* Buffer, uint32 BufferSize, FGeometryCacheMeshData& OutMeshData) override;
+	virtual FGeometryCacheCodecRenderStateBase *CreateRenderState() override;
+
+#if WITH_EDITORONLY_DATA
+	virtual void InitializeEncoder(float InVertexQuantizationPrecision, int32 InUVQuantizationBitRange);
+	virtual void BeginCoding(TArray<FStreamedGeometryCacheChunk> &AppendChunksTo) override;
+	virtual void EndCoding() override;
+	virtual void CodeFrame(const FGeometryCacheCodecEncodeArguments& Args) override;	
+#endif
+
+private:
+	ICodecDecoder* Decoder;
+
+#if WITH_EDITORONLY_DATA
+	ICodecEncoder* Encoder;
+
+	struct FEncoderData
+	{
+		int32 CurrentChunkId; // Index in the AppendChunksTo list of the chunk we are working on
+		TArray<FStreamedGeometryCacheChunk>* AppendChunksTo; // Any chunks this codec creates will be appended to this list
+	};
+	FEncoderData EncoderData;
+	int32 NextContextId;
+#endif
+};
+
diff --git a/Engine/Plugins/Runtime/GeometryCache/Source/GeometryCache/Classes/GeometryCacheTrackFlipbookAnimation.h b/Engine/Plugins/Runtime/GeometryCache/Source/GeometryCache/Classes/GeometryCacheTrackFlipbookAnimation.h
new file mode 100644
--- /dev/null
+++ b/Engine/Plugins/Runtime/GeometryCache/Source/GeometryCache/Classes/GeometryCacheTrackFlipbookAnimation.h
@@ -0,0 +1,49 @@
+// Copyright Epic Games, Inc. All Rights Reserved.
+#pragma once
+
+#include "CoreMinimal.h"
+#include "UObject/ObjectMacros.h"
+#include "UObject/ScriptMacros.h"
+#include "GeometryCacheTrack.h"
+#include "GeometryCacheMeshData.h"
+
+#include "GeometryCacheTrackFlipbookAnimation.generated.h"
+
+/** Derived GeometryCacheTrack class, used for Transform animation. */
+UCLASS(collapsecategories, hidecategories = Object, BlueprintType, config = Engine, deprecated)
+class GEOMETRYCACHE_API UDEPRECATED_GeometryCacheTrack_FlipbookAnimation : public UGeometryCacheTrack
+{
+	GENERATED_UCLASS_BODY()
+
+	virtual ~UDEPRECATED_GeometryCacheTrack_FlipbookAnimation();
+
+	//~ Begin UObject Interface.
+	virtual void GetResourceSizeEx(FResourceSizeEx& CumulativeResourceSize) override;
+	virtual void BeginDestroy() override;
+	//~ End UObject Interface.
+
+	//~ Begin UGeometryCacheTrack Interface.
+	virtual const bool UpdateMeshData(const float Time, const bool bLooping, int32& InOutMeshSampleIndex, FGeometryCacheMeshData*& OutMeshData) override;
+	virtual const float GetMaxSampleTime() const override;
+	//~ End UGeometryCacheTrack Interface.
+
+	/**
+	* Add a GeometryCacheMeshData sample to the Track
+	*
+	* @param MeshData - Holds the mesh data for the specific sample
+	* @param SampleTime - SampleTime for the specific sample being added
+	* @return void
+	*/
+	UFUNCTION()
+	void AddMeshSample(const FGeometryCacheMeshData& MeshData, const float SampleTime);
+
+private:
+	/** Number of Mesh Sample within this track */
+	UPROPERTY(VisibleAnywhere, Category = GeometryCache)
+	uint32 NumMeshSamples;
+
+	/** Stored data for each Mesh sample */
+	TArray<FGeometryCacheMeshData> MeshSamples;
+	TArray<float> MeshSampleTimes;
+};
+
diff --git a/Engine/Plugins/Runtime/GeometryCache/Source/GeometryCache/Classes/GeometryCacheTrackTransformAnimation.h b/Engine/Plugins/Runtime/GeometryCache/Source/GeometryCache/Classes/GeometryCacheTrackTransformAnimation.h
new file mode 100644
--- /dev/null
+++ b/Engine/Plugins/Runtime/GeometryCache/Source/GeometryCache/Classes/GeometryCacheTrackTransformAnimation.h
@@ -0,0 +1,35 @@
+// Copyright Epic Games, Inc. All Rights Reserved.
+#pragma once
+
+#include "CoreMinimal.h"
+#include "UObject/ObjectMacros.h"
+#include "UObject/ScriptMacros.h"
+#include "GeometryCacheTrack.h"
+#include "GeometryCacheMeshData.h"
+
+#include "GeometryCacheTrackTransformAnimation.generated.h"
+
+/** Derived GeometryCacheTrack class, used for Transform animation. */
+UCLASS(collapsecategories, hidecategories = Object, BlueprintType, config = Engine, deprecated)
+class GEOMETRYCACHE_API UDEPRECATED_GeometryCacheTrack_TransformAnimation : public UGeometryCacheTrack
+{
+	GENERATED_UCLASS_BODY()
+		
+	//~ Begin UObject Interface.
+	virtual void GetResourceSizeEx(FResourceSizeEx& CumulativeResourceSize) override;
+	//~ End UObject Interface.
+	
+	//~ Begin UGeometryCacheTrack Interface.
+	virtual const bool UpdateMeshData(const float Time, const bool bLooping, int32& InOutMeshSampleIndex, FGeometryCacheMeshData*& OutMeshData) override;
+	//~ End UGeometryCacheTrack Interface.
+
+	/**
+	* Sets/updates the MeshData for this track
+	*
+	* @param NewMeshData - GeometryCacheMeshData instance later used as the rendered mesh	
+	*/
+	UFUNCTION()
+	void SetMesh(const FGeometryCacheMeshData& NewMeshData);	
+private:
+	FGeometryCacheMeshData MeshData;
+};
diff --git a/Engine/Plugins/Runtime/GeometryCache/Source/GeometryCache/Classes/GeometryCacheTrackTransformGroupAnimation.h b/Engine/Plugins/Runtime/GeometryCache/Source/GeometryCache/Classes/GeometryCacheTrackTransformGroupAnimation.h
new file mode 100644
--- /dev/null
+++ b/Engine/Plugins/Runtime/GeometryCache/Source/GeometryCache/Classes/GeometryCacheTrackTransformGroupAnimation.h
@@ -0,0 +1,35 @@
+// Copyright Epic Games, Inc. All Rights Reserved.
+#pragma once
+
+#include "CoreMinimal.h"
+#include "UObject/ObjectMacros.h"
+#include "UObject/ScriptMacros.h"
+#include "GeometryCacheTrack.h"
+#include "GeometryCacheMeshData.h"
+
+#include "GeometryCacheTrackTransformGroupAnimation.generated.h"
+
+/** Derived GeometryCacheTrack class, used for Transform animation. */
+UCLASS(collapsecategories, hidecategories = Object, BlueprintType, config = Engine, deprecated)
+class GEOMETRYCACHE_API UDEPRECATED_GeometryCacheTrack_TransformGroupAnimation : public UGeometryCacheTrack
+{
+	GENERATED_UCLASS_BODY()
+		
+	// Begin UObject interface.
+	virtual void GetResourceSizeEx(FResourceSizeEx& CumulativeResourceSize) override;
+	// End UObject interface.
+	
+	// Begin UGeometryCacheTrack interface.
+	virtual const bool UpdateMeshData(const float Time, const bool bLooping, int32& InOutMeshSampleIndex, FGeometryCacheMeshData*& OutMeshData) override;
+	// End UGeometryCacheTrack interface.
+
+	/**
+	* Sets/updates the MeshData for this track
+	*
+	* @param NewMeshData - GeometryCacheMeshData instance later used as the rendered mesh	
+	*/
+	UFUNCTION()
+	void SetMesh(const FGeometryCacheMeshData& NewMeshData);	
+private:
+	FGeometryCacheMeshData MeshData;
+};
diff --git a/Engine/Plugins/Runtime/GeometryCache/Source/GeometryCache/Public/GeometryCacheSceneProxy.h b/Engine/Plugins/Runtime/GeometryCache/Source/GeometryCache/Public/GeometryCacheSceneProxy.h
new file mode 100644
--- /dev/null
+++ b/Engine/Plugins/Runtime/GeometryCache/Source/GeometryCache/Public/GeometryCacheSceneProxy.h
@@ -0,0 +1,411 @@
+// Copyright Epic Games, Inc. All Rights Reserved.
+#pragma once
+
+#include "CoreMinimal.h"
+#include "HitProxies.h"
+#include "Stats/Stats.h"
+#include "Containers/ResourceArray.h"
+#include "RenderResource.h"
+#include "PrimitiveViewRelevance.h"
+#include "PrimitiveSceneProxy.h"
+#include "Materials/MaterialInterface.h"
+#include "GeometryCacheVertexFactory.h"
+#include "DynamicMeshBuilder.h"
+#include "StaticMeshResources.h"
+#include "Logging/LogMacros.h"
+
+class FMeshElementCollector;
+struct FGeometryCacheMeshData;
+class FGeometryCacheTrackStreamableRenderResource;
+struct FVisibilitySample;
+class UGeometryCacheTrack;
+
+/** Resource array to pass  */
+class GEOMETRYCACHE_API FGeomCacheVertexBuffer : public FVertexBuffer
+{
+public:
+
+	void Init(int32 InSizeInBytes)
+	{
+		check(this->IsInitialized() == false);
+		SizeInBytes = InSizeInBytes;
+	}
+
+	/* Create on rhi thread. Uninitialized with NumVertices space.*/
+	virtual void InitRHI() override;
+
+	virtual void ReleaseRHI() override;
+
+	/**
+	 * Sugar function to update from a typed array.
+	 */
+	template<class DataType> void Update(const TArray<DataType>& Vertices)
+	{
+		int32 InSize = Vertices.Num() * sizeof(DataType);
+		UpdateRaw(Vertices.GetData(), InSize, 1, 1);
+	}
+
+	void UpdatePositionsOnly(const TArray<FDynamicMeshVertex>& Vertices)
+	{
+		const uint32 PositionOffset = STRUCT_OFFSET(FDynamicMeshVertex, Position);
+		const uint32 PositionSize = sizeof(((FDynamicMeshVertex*)nullptr)->Position);
+		UpdateRaw(Vertices.GetData() + PositionOffset, Vertices.Num(), PositionSize, sizeof(FDynamicMeshVertex));
+	}
+
+	void UpdateExceptPositions(const TArray<FDynamicMeshVertex>& Vertices)
+	{
+		const uint32 PositionSize = sizeof(((FDynamicMeshVertex*)nullptr)->Position);
+		const uint32 PositionOffset = STRUCT_OFFSET(FDynamicMeshVertex, Position);
+
+		static_assert(PositionOffset == 0, "Expecting position to be the first struct member");
+		static_assert(PositionSize == STRUCT_OFFSET(FDynamicMeshVertex, TextureCoordinate), "Expecting the texture coordinate to immediately follow the Position");
+
+		UpdateRaw((int8*)Vertices.GetData() + PositionSize, Vertices.Num(), sizeof(FDynamicMeshVertex) - PositionSize, sizeof(FDynamicMeshVertex));
+	}
+
+	/**
+	 * Update the raw contents of the buffer, possibly reallocate if needed.
+	 */
+	void UpdateRaw(const void* Data, int32 NumItems, int32 ItemSizeBytes, int32 ItemStrideBytes);
+
+	/**
+	 * Resize the buffer but don't initialize it with any data.
+	 */
+	void UpdateSize(int32 NewSizeInBytes);
+
+	/**
+	* Resize the buffer but don't initialize it with any data.
+	*/
+	template<class DataType> void UpdateSizeTyped(int32 NewSizeInElements)
+	{
+		UpdateSize(sizeof(DataType) * NewSizeInElements);
+	}
+
+	/**
+	 * Get the current size of the buffer
+	 */
+	unsigned GetSizeInBytes() { return SizeInBytes; }
+
+	virtual FString GetFriendlyName() const override { return TEXT("FGeomCacheVertexBuffer"); }
+
+	FRHIShaderResourceView* GetBufferSRV() const { return BufferSRV; }
+
+protected:
+	int32 SizeInBytes;
+	FShaderResourceViewRHIRef BufferSRV;
+};
+
+class GEOMETRYCACHE_API FGeomCacheTangentBuffer : public FGeomCacheVertexBuffer
+{
+public:
+	virtual void InitRHI() override;
+};
+
+class GEOMETRYCACHE_API FGeomCacheColorBuffer : public FGeomCacheVertexBuffer
+{
+public:
+	virtual void InitRHI() override;
+};
+
+/** Index Buffer */
+class GEOMETRYCACHE_API FGeomCacheIndexBuffer : public FIndexBuffer
+{
+public:
+	int32 NumAllocatedIndices = 0; // Total allocated GPU index buffer size in elements
+	int32 NumValidIndices = 0; // Current valid data region of the index buffer (may be smaller than allocated buffer)
+
+	/* Create on rhi thread. Uninitialized with NumIndices space.*/
+	virtual void InitRHI() override;
+
+	virtual void ReleaseRHI() override;
+
+	/**
+		Update the data and possibly reallocate if needed.
+	*/
+	void Update(const TArray<uint32>& Indices);
+
+	void UpdateSizeOnly(int32 NewNumIndices);
+
+	unsigned SizeInBytes() { return NumAllocatedIndices * sizeof(uint32); }
+
+	FRHIShaderResourceView* GetBufferSRV() const { return BufferSRV; }
+
+protected:
+	FShaderResourceViewRHIRef BufferSRV;
+};
+
+/** Vertex Factory */
+class GEOMETRYCACHE_API FGeomCacheVertexFactory : public FGeometryCacheVertexVertexFactory
+{
+public:
+
+	FGeomCacheVertexFactory(ERHIFeatureLevel::Type InFeatureLevel);
+
+	/** Init function that should only be called on render thread. */
+	void Init_RenderThread(const FVertexBuffer* PositionBuffer, const FVertexBuffer* MotionBlurDataBuffer, const FVertexBuffer* TangentXBuffer, const FVertexBuffer* TangentZBuffer, const FVertexBuffer* TextureCoordinateBuffer, const FVertexBuffer* ColorBuffer);
+
+	/** Init function that can be called on any thread, and will do the right thing (enqueue command if called on main thread) */
+	void Init(const FVertexBuffer* PositionBuffer, const FVertexBuffer* MotionBlurDataBuffer, const FVertexBuffer* TangentXBuffer, const FVertexBuffer* TangentZBuffer, const FVertexBuffer* TextureCoordinateBuffer, const FVertexBuffer* ColorBuffer);
+};
+
+/**
+ * This the track proxy has some "double double buffering" going on.
+ * First we keep two mesh frames. The one just before the current time and the one just after the current time. This is the full mesh and
+ * we interpolate between it to derive the actual mesh for the exact time we're at.
+ * Secondly we have two position buffers. The one for the current rendered frame and the one from the previous rendered frame (this is not the same as
+ * the mesh frame, the mesh may be at say 10 fps then get interpolated to 60 fps rendered frames)
+ */
+class GEOMETRYCACHE_API FGeomCacheTrackProxy
+{
+public:
+
+	FGeomCacheTrackProxy(ERHIFeatureLevel::Type InFeatureLevel)
+		: VertexFactory(InFeatureLevel/*, "FGeomCacheTrackProxy"*/)
+	{}
+
+	virtual ~FGeomCacheTrackProxy() {}
+
+	/**
+	 * Update the SampleIndex and MeshData for a given time
+	 *
+	 * @param Time - (Elapsed)Time to check against
+	 * @param bLooping - Whether or not the animation is being played in a loop
+	 * @param InOutMeshSampleIndex - Hold the MeshSampleIndex and will be updated if changed according to the Elapsed Time
+	 * @param OutMeshData - Will hold the new MeshData if the SampleIndex changed
+	 * @return true if the SampleIndex and MeshData were updated
+	 */
+	virtual bool UpdateMeshData(float Time, bool bLooping, int32& InOutMeshSampleIndex, FGeometryCacheMeshData& OutMeshData);
+
+	/**
+	 * Get the MeshData for a given SampleIndex
+	 *
+	 * @param SampleIndex - The sample index at which to retrieve the MeshData
+	 * @param OutMeshData - Will hold the MeshData if it could be retrieved
+	 * @return true if the MeshData was retrieved successfully
+	 */
+	virtual bool GetMeshData(int32 SampleIndex, FGeometryCacheMeshData& OutMeshData);
+
+	/**
+	 * Check if the topology of two given SampleIndexes are compatible (ie. same topology)
+	 *
+	 * @param SampleIndexA - The first sample index to compare the topology
+	 * @param SampleIndexB - The second sample index to compare the topology
+	 * @return true if the topology is the same
+	 */
+	virtual bool IsTopologyCompatible(int32 SampleIndexA, int32 SampleIndexB);
+
+	/**
+	 * Get the VisibilitySample for a given time
+	 *
+	 * @param Time - (Elapsed)Time to check against
+	 * @param bLooping - Whether or not the animation is being played in a loop
+	 * @return the VisibilitySample that corresponds to the given time
+	 */
+	virtual const FVisibilitySample& GetVisibilitySample(float Time, const bool bLooping) const;
+
+	/**
+	 * Find the two frames closest to the given time
+	 * InterpolationFactor gives the position of the requested time slot between the two returned frames.
+	 * 0.0 => We are very close to OutFrameIndex
+	 * 1.0 => We are very close to OutNextFrameIndex
+	 * If bIsPlayingBackwards it will return exactly the same indexes but in the reversed order. The
+	 * InterpolationFactor will also be updated accordingly
+	 *
+	 * @param Time - (Elapsed)Time to check against
+	 * @param bLooping - Whether or not the animation is being played in a loop
+	 * @param bIsPlayingBackwards - Whether the animation is playing backwards or forwards
+	 * @param OutFrameIndex - The closest frame index that corresponds to the given time
+	 * @param OutNextFrameIndex - The frame index that follows OutFrameIndex
+	 * @param InterpolationFactor - The interpolation value between the two frame times
+	 */
+	virtual void FindSampleIndexesFromTime(float Time, bool bLooping, bool bIsPlayingBackwards, int32& OutFrameIndex, int32& OutNextFrameIndex, float& InInterpolationFactor);
+
+	/**
+	 * Initialize the render resources. Must be called before the render resources are used.
+	 *
+	 * @param NumVertices - The initial number of vertices to initialize the buffers with. Must be greater than 0
+	 * @param NumIndices - The initial number of indices to initialize the buffers with. Must be greater than 0
+	 */
+	virtual void InitRenderResources(int32 NumVertices, int32 NumIndices);
+
+	/** MeshData storing information used for rendering this Track */
+	FGeometryCacheMeshData* MeshData;
+	FGeometryCacheMeshData* NextFrameMeshData;
+
+	/** Frame numbers corresponding to MeshData, NextFrameMeshData */
+	int32 FrameIndex;
+	int32 NextFrameIndex;
+	int32 PreviousFrameIndex;
+	float InterpolationFactor;
+	float PreviousInterpolationFactor;
+	float SubframeInterpolationFactor;
+
+	/** Material applied to this Track */
+	TArray<UMaterialInterface*> Materials;
+
+	/** Vertex buffers for this Track. There are two position buffers which we double buffer between, current frame and last frame*/
+	FGeomCacheVertexBuffer PositionBuffers[2];
+	uint32 PositionBufferFrameIndices[2]; // Frame indexes of the positions in the position buffer 
+	float PositionBufferFrameTimes[2]; // Exact time after interpolation of the positions in the position buffer.
+	uint32 CurrentPositionBufferIndex; // CurrentPositionBufferIndex%2  is the last updated position buffer
+
+	FGeomCacheTangentBuffer TangentXBuffer;
+	FGeomCacheTangentBuffer TangentZBuffer;
+	FGeomCacheVertexBuffer TextureCoordinatesBuffer;
+	FGeomCacheColorBuffer ColorBuffer;
+
+	/** Index buffer for this Track */
+	FGeomCacheIndexBuffer IndexBuffer;
+
+	/** Vertex factory for this Track */
+	FGeomCacheVertexFactory VertexFactory;
+
+	/** World Matrix for this Track */
+	FMatrix WorldMatrix;
+
+	/** The GeometryCacheTrack to which the proxy is associated */
+	UGeometryCacheTrack* Track;
+
+	int32 UploadedSampleIndex;
+
+	/** Flag to indicate which frame mesh data was selected during the update */
+	bool bNextFrameMeshDataSelected;
+
+	bool bResourcesInitialized;
+
+#if RHI_RAYTRACING
+	FRayTracingGeometry RayTracingGeometry;
+#endif
+};
+
+/** Procedural mesh scene proxy */
+class GEOMETRYCACHE_API FGeometryCacheSceneProxy : public FPrimitiveSceneProxy
+{
+public:
+	SIZE_T GetTypeHash() const override;
+
+	FGeometryCacheSceneProxy(class UGeometryCacheComponent* Component);
+	FGeometryCacheSceneProxy(class UGeometryCacheComponent* Component, TFunction<FGeomCacheTrackProxy*()> TrackProxyCreator);
+
+	virtual ~FGeometryCacheSceneProxy();
+
+	// Begin FPrimitiveSceneProxy interface.
+#if WITH_EDITOR
+	virtual HHitProxy* CreateHitProxies(UPrimitiveComponent* Component, TArray<TRefCountPtr<HHitProxy> >& OutHitProxies) override;
+#endif
+	virtual void GetDynamicMeshElements(const TArray<const FSceneView*>& Views, const FSceneViewFamily& ViewFamily, uint32 VisibilityMap, FMeshElementCollector& Collector) const override;
+	virtual FPrimitiveViewRelevance GetViewRelevance(const FSceneView* View) const override;
+
+	virtual bool CanBeOccluded() const override;
+	virtual bool IsUsingDistanceCullFade() const override;
+	virtual uint32 GetMemoryFootprint(void) const;
+	uint32 GetAllocatedSize(void) const;
+	// End FPrimitiveSceneProxy interface.
+
+	void UpdateAnimation(float NewTime, bool bLooping, bool bIsPlayingBackwards, float PlaybackSpeed, float MotionVectorScale);
+
+	/** Update world matrix for specific section */
+	void UpdateSectionWorldMatrix(const int32 SectionIndex, const FMatrix& WorldMatrix);
+	/** Update vertex buffer for specific section */
+	void UpdateSectionVertexBuffer(const int32 SectionIndex, FGeometryCacheMeshData* MeshData);
+	/** Update index buffer for specific section */
+	void UpdateSectionIndexBuffer(const int32 SectionIndex, const TArray<uint32>& Indices);
+
+	/** Clears the Sections array*/
+	void ClearSections();
+
+#if RHI_RAYTRACING
+	virtual void GetDynamicRayTracingInstances(FRayTracingMaterialGatheringContext& Context, TArray<FRayTracingInstance>& OutRayTracingInstances) override final;
+	virtual bool IsRayTracingRelevant() const override { return true; }
+	virtual bool HasRayTracingRepresentation() const override { return true; }
+#endif
+
+	const TArray<FGeomCacheTrackProxy*>& GetTracks() { return Tracks; }
+
+private:
+	void FrameUpdate() const;
+
+	void CreateMeshBatch(
+		const FGeomCacheTrackProxy* TrackProxy,
+		const struct FGeometryCacheMeshBatchInfo& BatchInfo,
+		class FGeometryCacheVertexFactoryUserDataWrapper& UserDataWrapper,
+		FDynamicPrimitiveUniformBuffer& DynamicPrimitiveUniformBuffer,
+		FMeshBatch& Mesh) const;
+
+private:
+	/** Array of Track Proxies */
+	TArray<FGeomCacheTrackProxy*> Tracks;
+
+#if WITH_EDITOR
+	TArray<FHitProxyId> HitProxyIds;
+#endif
+
+	/** Scratch memory for frame update - do not use directly. */
+	struct FScratchMemory
+	{
+		TArray<FVector3f> InterpolatedPositions;
+		TArray<FPackedNormal> InterpolatedTangentX;
+		TArray<FPackedNormal> InterpolatedTangentZ;
+		TArray<FVector2f> InterpolatedUVs;
+		TArray<FColor> InterpolatedColors;
+		TArray<FVector3f> InterpolatedMotionVectors;
+
+		void Prepare(SIZE_T NumVertices, bool bHasMotionVectors)
+		{
+			// Clear entries but keep allocations.
+			InterpolatedPositions.Reset();
+			InterpolatedTangentX.Reset();
+			InterpolatedTangentZ.Reset();
+			InterpolatedUVs.Reset();
+			InterpolatedColors.Reset();
+			InterpolatedMotionVectors.Reset();
+
+			// Make sure our capacity fits the requested vertex count
+			InterpolatedPositions.Reserve(NumVertices);
+			InterpolatedTangentX.Reserve(NumVertices);
+			InterpolatedTangentZ.Reserve(NumVertices);
+			InterpolatedUVs.Reserve(NumVertices);
+			InterpolatedColors.Reserve(NumVertices);
+
+			InterpolatedPositions.AddUninitialized(NumVertices);
+			InterpolatedTangentX.AddUninitialized(NumVertices);
+			InterpolatedTangentZ.AddUninitialized(NumVertices);
+			InterpolatedUVs.AddUninitialized(NumVertices);
+			InterpolatedColors.AddUninitialized(NumVertices);
+
+			if (bHasMotionVectors)
+			{
+				InterpolatedMotionVectors.Reserve(NumVertices);
+				InterpolatedMotionVectors.AddUninitialized(NumVertices);
+			}
+		}
+
+		void Empty()
+		{
+			// Clear entries but and release memory.
+			InterpolatedPositions.Empty();
+			InterpolatedTangentX.Empty();
+			InterpolatedTangentZ.Empty();
+			InterpolatedUVs.Empty();
+			InterpolatedColors.Empty();
+			InterpolatedMotionVectors.Empty();
+		}
+	}
+	mutable Scratch;
+
+	uint32 UpdatedFrameNum;
+	float Time;
+	float PlaybackSpeed;
+	float MotionVectorScale;
+
+	bool bOverrideWireframeColor = false;
+	FLinearColor WireframeOverrideColor = FLinearColor::Green;
+
+	FMaterialRelevance MaterialRelevance;
+	uint32 bLooping : 1;
+	uint32 bIsPlayingBackwards : 1;
+	uint32 bExtrapolateFrames : 1;
+
+	/** Function used to create a new track proxy at construction */
+	TFunction<FGeomCacheTrackProxy*()> CreateTrackProxy;
+};
diff --git a/Engine/Plugins/Runtime/GeometryCache/Source/GeometryCacheEd/Private/AssetTypeActions_GeometryCache.cpp b/Engine/Plugins/Runtime/GeometryCache/Source/GeometryCacheEd/Private/AssetTypeActions_GeometryCache.cpp
new file mode 100644
--- /dev/null
+++ b/Engine/Plugins/Runtime/GeometryCache/Source/GeometryCacheEd/Private/AssetTypeActions_GeometryCache.cpp
@@ -0,0 +1,69 @@
+// Copyright Epic Games, Inc. All Rights Reserved.
+
+#include "AssetTypeActions_GeometryCache.h"
+#include "EditorFramework/AssetImportData.h"
+#include "ThumbnailRendering/SceneThumbnailInfo.h"
+#include "GeometryCache.h"
+#include "ToolMenus.h"
+
+FText FAssetTypeActions_GeometryCache::GetName() const
+{
+	return NSLOCTEXT("AssetTypeActions", "AssetTypeActions_GeometryCache", "GeometryCache");
+}
+
+FColor FAssetTypeActions_GeometryCache::GetTypeColor() const
+{
+	return FColor(0, 255, 255);
+}
+
+UClass* FAssetTypeActions_GeometryCache::GetSupportedClass() const
+{
+	return UGeometryCache::StaticClass();
+}
+
+bool FAssetTypeActions_GeometryCache::HasActions(const TArray<UObject*>& InObjects) const
+{
+	return false;
+}
+
+void FAssetTypeActions_GeometryCache::GetActions(const TArray<UObject*>& InObjects, FToolMenuSection& Section)
+{
+	FAssetTypeActions_Base::GetActions(InObjects, Section);
+}
+
+void FAssetTypeActions_GeometryCache::OpenAssetEditor(const TArray<UObject*>& InObjects, TSharedPtr<class IToolkitHost> EditWithinLevelEditor /*= TSharedPtr<IToolkitHost>()*/)
+{
+	FAssetTypeActions_Base::OpenAssetEditor(InObjects, EditWithinLevelEditor);
+}
+
+uint32 FAssetTypeActions_GeometryCache::GetCategories()
+{
+	return EAssetTypeCategories::Animation;
+}
+
+class UThumbnailInfo* FAssetTypeActions_GeometryCache::GetThumbnailInfo(UObject* Asset) const
+{
+	UGeometryCache* GeometryCache = CastChecked<UGeometryCache>(Asset);
+	UThumbnailInfo* ThumbnailInfo = GeometryCache->ThumbnailInfo;
+	if (ThumbnailInfo == NULL)
+	{
+		ThumbnailInfo = NewObject<USceneThumbnailInfo>(GeometryCache, NAME_None, RF_Transactional);
+		GeometryCache->ThumbnailInfo = ThumbnailInfo;
+	}
+
+	return ThumbnailInfo;
+}
+
+bool FAssetTypeActions_GeometryCache::IsImportedAsset() const
+{
+	return true;
+}
+
+void FAssetTypeActions_GeometryCache::GetResolvedSourceFilePaths(const TArray<UObject*>& TypeAssets, TArray<FString>& OutSourceFilePaths) const
+{
+	for (auto& Asset : TypeAssets)
+	{
+		const auto GeometryCache = CastChecked<UGeometryCache>(Asset);
+		GeometryCache->AssetImportData->ExtractFilenames(OutSourceFilePaths);
+	}
+}
