Streaming GeometryCache:
* Added a cvar GeometryCache.Streamer.BlockTillFinishStreaming to wait for the GeometryCache streamer to finish streaming the requested frames; for use with the Movie Render Queue.

diff --git a/Engine/Plugins/Runtime/GeometryCache/Source/GeometryCacheStreamer/Private/GeometryCacheStreamer.cpp b/Engine/Plugins/Runtime/GeometryCache/Source/GeometryCacheStreamer/Private/GeometryCacheStreamer.cpp
--- a/Engine/Plugins/Runtime/GeometryCache/Source/GeometryCacheStreamer/Private/GeometryCacheStreamer.cpp
+++ b/Engine/Plugins/Runtime/GeometryCache/Source/GeometryCacheStreamer/Private/GeometryCacheStreamer.cpp
@@ -7,6 +7,7 @@
 #include "GeometryCacheMeshData.h"
 #include "GeometryCacheModule.h"
 #include "GeometryCacheStreamerSettings.h"
+#include "HAL/PlatformProcess.h"
 #include "IGeometryCacheStream.h"
 #include "Stats/Stats.h"
 #include "Widgets/Notifications/SNotificationList.h"
@@ -24,6 +25,12 @@ static FAutoConsoleVariableRef CVarGeometryCacheStreamerShowNotification(
 	GShowGeometryCacheStreamerNotification,
 	TEXT("Show notification while the GeometryCache streamer is streaming data"));
 
+static bool GGeoCacheStreamerBlockTillFinishStreaming = false;
+static FAutoConsoleVariableRef CVarBlockTillFinishStreaming(
+	TEXT("GeometryCache.Streamer.BlockTillFinishStreaming"),
+	GGeoCacheStreamerBlockTillFinishStreaming,
+	TEXT("Force the GeometryCache streamer to block until it has finished streaming all the requested frames"));
+
 class FGeometryCacheStreamer : public IGeometryCacheStreamer
 {
 public:
@@ -83,68 +90,81 @@ FGeometryCacheStreamer::~FGeometryCacheStreamer()
 
 void FGeometryCacheStreamer::Tick(float Time)
 {
-	// The sole purpose of the Streamer is to schedule the streams for read at every engine loop
-	BalanceStreams();
-
-	// First step is to process the results from the previously scheduled read requests
-	// to figure out the number of reads still in progress
+	int32 NumStreams = 0;
 	int32 NumFramesToStream = 0;
-	for (const auto& Pair : TracksToStreams)
+	bool bWaitForFinish = false;
+	do 
 	{
-		TArray<int32> FramesCompleted;
-		IGeometryCacheStream* Stream = Pair.Value;
-		Stream->UpdateRequestStatus(FramesCompleted);
-		NumReads -= FramesCompleted.Num();
-		NumFramesToStream += Stream->GetNumFramesNeeded();
-	}
+		// The sole purpose of the Streamer is to schedule the streams for read at every engine loop
+		BalanceStreams();
 
-	// Now, schedule new read requests according to the number of concurrent reads available
-	const int32 NumStreams = TracksToStreams.Num();
-	int32 AvailableReads = MaxReads - NumReads;
-	if (NumStreams > 0 && AvailableReads > 0)
-	{
-		TArray<IGeometryCacheStream*> Streams;
-		TracksToStreams.GenerateValueArray(Streams);
-
-		// Streams are checked round robin until there are no more reads available
-		// or no stream can handle more read requests
-		// Note that the round robin starts from where it left off in the previous Tick
-		TBitArray<> StreamsToCheck(true, NumStreams);
-		for (; AvailableReads > 0 && StreamsToCheck.Contains(true); ++CurrentIndex)
+		// First step is to process the results from the previously scheduled read requests
+		// to figure out the number of reads still in progress
+		NumFramesToStream = 0;
+		for (const auto& Pair : TracksToStreams)
 		{
-			// Handle looping but also the case where the number of streams has decreased
-			if (CurrentIndex >= NumStreams)
-			{
-				CurrentIndex = 0;
-			}
+			TArray<int32> FramesCompleted;
+			IGeometryCacheStream* Stream = Pair.Value;
+			Stream->UpdateRequestStatus(FramesCompleted);
+			NumReads -= FramesCompleted.Num();
+			NumFramesToStream += Stream->GetNumFramesNeeded();
+		}
 
-			if (!StreamsToCheck[CurrentIndex])
+		// Now, schedule new read requests according to the number of concurrent reads available
+		NumStreams = TracksToStreams.Num();
+		int32 AvailableReads = MaxReads - NumReads;
+		if (NumStreams > 0 && AvailableReads > 0)
+		{
+			TArray<IGeometryCacheStream*> Streams;
+			TracksToStreams.GenerateValueArray(Streams);
+
+			// Streams are checked round robin until there are no more reads available
+			// or no stream can handle more read requests
+			// Note that the round robin starts from where it left off in the previous Tick
+			TBitArray<> StreamsToCheck(true, NumStreams);
+			for (; AvailableReads > 0 && StreamsToCheck.Contains(true); ++CurrentIndex)
 			{
-				continue;
-			}
+				// Handle looping but also the case where the number of streams has decreased
+				if (CurrentIndex >= NumStreams)
+				{
+					CurrentIndex = 0;
+				}
 
-			IGeometryCacheStream* Stream = Streams[CurrentIndex];
-			if (Stream->GetNumFramesNeeded() > 0)
-			{
-				if (Stream->RequestFrameData())
+				if (!StreamsToCheck[CurrentIndex])
+				{
+					continue;
+				}
+
+				IGeometryCacheStream* Stream = Streams[CurrentIndex];
+				if (Stream->GetNumFramesNeeded() > 0)
 				{
-					// Stream was able to handle the read request so there's one less available
-					++NumReads;
-					--AvailableReads;
+					if (Stream->RequestFrameData())
+					{
+						// Stream was able to handle the read request so there's one less available
+						++NumReads;
+						--AvailableReads;
+					}
+					else
+					{
+						// Stream cannot handle more read request, don't need to check it again
+						StreamsToCheck[CurrentIndex] = false;
+					}
 				}
 				else
 				{
-					// Stream cannot handle more read request, don't need to check it again
+					// Stream doesn't need any frame to be read, don't need to check it again
 					StreamsToCheck[CurrentIndex] = false;
 				}
 			}
-			else
-			{
-				// Stream doesn't need any frame to be read, don't need to check it again
-				StreamsToCheck[CurrentIndex] = false;
-			}
 		}
-	}
+
+		// If this cvar is enabled, tick until all frames have finished streaming
+		bWaitForFinish = GGeoCacheStreamerBlockTillFinishStreaming && (NumFramesToStream > 0);
+		if (bWaitForFinish)
+		{
+			FPlatformProcess::Sleep(.1f);
+		}
+	} while (bWaitForFinish);
 
 	INC_DWORD_STAT_BY(STAT_GeometryCacheStream_Count, NumStreams);
 
@@ -266,4 +286,4 @@ IGeometryCacheStreamer& IGeometryCacheStreamer::Get()
 	return Streamer;
 }
 
-#undef LOCTEXT_NAMESPACE
\ No newline at end of file
+#undef LOCTEXT_NAMESPACE
