- Remove dead functions

#rnx

diff --git a/Engine/Plugins/FX/Niagara/Source/Niagara/Private/NiagaraLightRendererProperties.cpp b/Engine/Plugins/FX/Niagara/Source/Niagara/Private/NiagaraLightRendererProperties.cpp
--- a/Engine/Plugins/FX/Niagara/Source/Niagara/Private/NiagaraLightRendererProperties.cpp
+++ b/Engine/Plugins/FX/Niagara/Source/Niagara/Private/NiagaraLightRendererProperties.cpp
@@ -163,15 +163,6 @@ void UNiagaraLightRendererProperties::GetRendererFeedback(const UNiagaraEmitter*
 	Super::GetRendererFeedback(InEmitter, OutErrors, OutWarnings, OutInfo);
 }
 
-bool UNiagaraLightRendererProperties::IsMaterialValidForRenderer(UMaterial* Material, FText& InvalidMessage)
-{
-	return true;
-}
-
-void UNiagaraLightRendererProperties::FixMaterial(UMaterial* Material)
-{
-}
-
 #endif // WITH_EDITORONLY_DATA
 
 #undef LOCTEXT_NAMESPACE
\ No newline at end of file
diff --git a/Engine/Plugins/FX/Niagara/Source/Niagara/Private/NiagaraMeshRendererProperties.cpp b/Engine/Plugins/FX/Niagara/Source/Niagara/Private/NiagaraMeshRendererProperties.cpp
--- a/Engine/Plugins/FX/Niagara/Source/Niagara/Private/NiagaraMeshRendererProperties.cpp
+++ b/Engine/Plugins/FX/Niagara/Source/Niagara/Private/NiagaraMeshRendererProperties.cpp
@@ -548,23 +548,6 @@ void UNiagaraMeshRendererProperties::PostLoad()
 }
 
 #if WITH_EDITORONLY_DATA
-bool UNiagaraMeshRendererProperties::IsMaterialValidForRenderer(UMaterial* Material, FText& InvalidMessage)
-{
-	if (Material->bUsedWithNiagaraMeshParticles == false)
-	{
-		InvalidMessage = NSLOCTEXT("NiagaraMeshRendererProperties", "InvalidMaterialMessage", "The material isn't marked as \"Used with Niagara Mesh particles\"");
-		return false;
-	}
-	return true;
-}
-
-void UNiagaraMeshRendererProperties::FixMaterial(UMaterial* Material)
-{
-	Material->Modify();
-	Material->bUsedWithNiagaraMeshParticles = true;
-	Material->ForceRecompileForRendering();
-}
-
 const TArray<FNiagaraVariable>& UNiagaraMeshRendererProperties::GetOptionalAttributes()
 {
 	static TArray<FNiagaraVariable> Attrs;
diff --git a/Engine/Plugins/FX/Niagara/Source/Niagara/Private/NiagaraRibbonRendererProperties.cpp b/Engine/Plugins/FX/Niagara/Source/Niagara/Private/NiagaraRibbonRendererProperties.cpp
--- a/Engine/Plugins/FX/Niagara/Source/Niagara/Private/NiagaraRibbonRendererProperties.cpp
+++ b/Engine/Plugins/FX/Niagara/Source/Niagara/Private/NiagaraRibbonRendererProperties.cpp
@@ -473,23 +473,6 @@ void UNiagaraRibbonRendererProperties::GetRendererFeedback(const UNiagaraEmitter
 }
 
 
-bool UNiagaraRibbonRendererProperties::IsMaterialValidForRenderer(UMaterial* InMaterial, FText& InvalidMessage)
-{
-	if (InMaterial->bUsedWithNiagaraRibbons == false)
-	{
-		InvalidMessage = NSLOCTEXT("NiagaraRibbonRendererProperties", "InvalidMaterialMessage", "The material isn't marked as \"Used with Niagara ribbons\"");
-		return false;
-	}
-	return true;
-}
-
-void UNiagaraRibbonRendererProperties::FixMaterial(UMaterial* InMaterial)
-{
-	InMaterial->Modify();
-	InMaterial->bUsedWithNiagaraRibbons = true;
-	InMaterial->ForceRecompileForRendering();
-}
-
 bool UNiagaraRibbonRendererProperties::CanEditChange(const FProperty* InProperty) const
 {
 
diff --git a/Engine/Plugins/FX/Niagara/Source/Niagara/Private/NiagaraSpriteRendererProperties.cpp b/Engine/Plugins/FX/Niagara/Source/Niagara/Private/NiagaraSpriteRendererProperties.cpp
--- a/Engine/Plugins/FX/Niagara/Source/Niagara/Private/NiagaraSpriteRendererProperties.cpp
+++ b/Engine/Plugins/FX/Niagara/Source/Niagara/Private/NiagaraSpriteRendererProperties.cpp
@@ -571,23 +571,6 @@ void UNiagaraSpriteRendererProperties::GetRendererTooltipWidgets(const FNiagaraE
 	}
 }
 
-bool UNiagaraSpriteRendererProperties::IsMaterialValidForRenderer(UMaterial* InMaterial, FText& InvalidMessage)
-{
-	if (InMaterial->bUsedWithNiagaraSprites == false)
-	{
-		InvalidMessage = NSLOCTEXT("NiagaraSpriteRendererProperties", "InvalidMaterialMessage", "The material isn't marked as \"Used with particle sprites\"");
-		return false;
-	}
-	return true;
-}
-
-void UNiagaraSpriteRendererProperties::FixMaterial(UMaterial* InMaterial)
-{
-	InMaterial->Modify();
-	InMaterial->bUsedWithNiagaraSprites = true;
-	InMaterial->ForceRecompileForRendering();
-}
-
 void UNiagaraSpriteRendererProperties::UpdateCutoutTexture()
 {
 	if (bUseMaterialCutoutTexture)
diff --git a/Engine/Plugins/FX/Niagara/Source/Niagara/Public/NiagaraLightRendererProperties.h b/Engine/Plugins/FX/Niagara/Source/Niagara/Public/NiagaraLightRendererProperties.h
--- a/Engine/Plugins/FX/Niagara/Source/Niagara/Public/NiagaraLightRendererProperties.h
+++ b/Engine/Plugins/FX/Niagara/Source/Niagara/Public/NiagaraLightRendererProperties.h
@@ -33,8 +33,6 @@ public:
 	virtual void GetUsedMaterials(const FNiagaraEmitterInstance* InEmitter, TArray<UMaterialInterface*>& OutMaterials) const override;
 	virtual bool IsSimTargetSupported(ENiagaraSimTarget InSimTarget) const override { return (InSimTarget == ENiagaraSimTarget::CPUSim); };
 #if WITH_EDITORONLY_DATA
-	virtual bool IsMaterialValidForRenderer(UMaterial* Material, FText& InvalidMessage) override;
-	virtual void FixMaterial(UMaterial* Material) override;
 	virtual const TArray<FNiagaraVariable>& GetOptionalAttributes() override;
 	virtual void GetRendererWidgets(const FNiagaraEmitterInstance* InEmitter, TArray<TSharedPtr<SWidget>>& OutWidgets, TSharedPtr<FAssetThumbnailPool> InThumbnailPool) const override;
 	virtual void GetRendererTooltipWidgets(const FNiagaraEmitterInstance* InEmitter, TArray<TSharedPtr<SWidget>>& OutWidgets, TSharedPtr<FAssetThumbnailPool> InThumbnailPool) const override;
diff --git a/Engine/Plugins/FX/Niagara/Source/Niagara/Public/NiagaraMeshRendererProperties.h b/Engine/Plugins/FX/Niagara/Source/Niagara/Public/NiagaraMeshRendererProperties.h
--- a/Engine/Plugins/FX/Niagara/Source/Niagara/Public/NiagaraMeshRendererProperties.h
+++ b/Engine/Plugins/FX/Niagara/Source/Niagara/Public/NiagaraMeshRendererProperties.h
@@ -180,8 +180,6 @@ public:
 	virtual bool IsSimTargetSupported(ENiagaraSimTarget InSimTarget) const override { return true; };
 	virtual bool PopulateRequiredBindings(FNiagaraParameterStore& InParameterStore) override;
 #if WITH_EDITORONLY_DATA
-	virtual bool IsMaterialValidForRenderer(UMaterial* Material, FText& InvalidMessage) override;
-	virtual void FixMaterial(UMaterial* Material) override;
 	virtual const TArray<FNiagaraVariable>& GetOptionalAttributes() override;
 	virtual void GetAdditionalVariables(TArray<FNiagaraVariableBase>& OutArray) const override;
 	virtual	void GetRendererWidgets(const FNiagaraEmitterInstance* InEmitter, TArray<TSharedPtr<SWidget>>& OutWidgets, TSharedPtr<FAssetThumbnailPool> InThumbnailPool) const override;
diff --git a/Engine/Plugins/FX/Niagara/Source/Niagara/Public/NiagaraRibbonRendererProperties.h b/Engine/Plugins/FX/Niagara/Source/Niagara/Public/NiagaraRibbonRendererProperties.h
--- a/Engine/Plugins/FX/Niagara/Source/Niagara/Public/NiagaraRibbonRendererProperties.h
+++ b/Engine/Plugins/FX/Niagara/Source/Niagara/Public/NiagaraRibbonRendererProperties.h
@@ -214,8 +214,6 @@ public:
 	bool PopulateRequiredBindings(FNiagaraParameterStore& InParameterStore);
 
 #if WITH_EDITOR
-	virtual bool IsMaterialValidForRenderer(UMaterial* Material, FText& InvalidMessage) override;
-	virtual void FixMaterial(UMaterial* Material);
 	virtual const TArray<FNiagaraVariable>& GetOptionalAttributes() override;
 	virtual void GetAdditionalVariables(TArray<FNiagaraVariableBase>& OutArray) const override;
 	virtual FNiagaraVariable GetBoundAttribute(const FNiagaraVariableAttributeBinding* Binding) const override;
diff --git a/Engine/Plugins/FX/Niagara/Source/Niagara/Public/NiagaraSpriteRendererProperties.h b/Engine/Plugins/FX/Niagara/Source/Niagara/Public/NiagaraSpriteRendererProperties.h
--- a/Engine/Plugins/FX/Niagara/Source/Niagara/Public/NiagaraSpriteRendererProperties.h
+++ b/Engine/Plugins/FX/Niagara/Source/Niagara/Public/NiagaraSpriteRendererProperties.h
@@ -126,8 +126,6 @@ public:
 	virtual bool IsSimTargetSupported(ENiagaraSimTarget InSimTarget) const override { return true; };
 	virtual bool PopulateRequiredBindings(FNiagaraParameterStore& InParameterStore)  override;
 #if WITH_EDITOR
-	virtual bool IsMaterialValidForRenderer(UMaterial* Material, FText& InvalidMessage) override;
-	virtual void FixMaterial(UMaterial* Material) override;
 	virtual const TArray<FNiagaraVariable>& GetOptionalAttributes() override;
 	virtual void GetAdditionalVariables(TArray<FNiagaraVariableBase>& OutArray) const override;
 	virtual void GetRendererWidgets(const FNiagaraEmitterInstance* InEmitter, TArray<TSharedPtr<SWidget>>& OutWidgets, TSharedPtr<FAssetThumbnailPool> InThumbnailPool) const override;
diff --git a/Engine/Plugins/Runtime/GeometryCache/Source/GeometryCache/Classes/NiagaraGeometryCacheRendererProperties.h b/Engine/Plugins/Runtime/GeometryCache/Source/GeometryCache/Classes/NiagaraGeometryCacheRendererProperties.h
--- a/Engine/Plugins/Runtime/GeometryCache/Source/GeometryCache/Classes/NiagaraGeometryCacheRendererProperties.h
+++ b/Engine/Plugins/Runtime/GeometryCache/Source/GeometryCache/Classes/NiagaraGeometryCacheRendererProperties.h
@@ -58,8 +58,6 @@ public:
 	virtual const FSlateBrush* GetStackIcon() const override;
 	virtual void GetRendererTooltipWidgets(const FNiagaraEmitterInstance* InEmitter, TArray<TSharedPtr<SWidget>>& OutWidgets, TSharedPtr<FAssetThumbnailPool> InThumbnailPool) const override;
 	virtual const TArray<FNiagaraVariable>& GetOptionalAttributes() override;
-	virtual bool IsMaterialValidForRenderer(UMaterial* Material, FText& InvalidMessage) override;
-	virtual void FixMaterial(UMaterial* Material) override;
 	virtual void GetRendererFeedback(const UNiagaraEmitter* InEmitter, TArray<FText>& OutErrors, TArray<FText>& OutWarnings, TArray<FText>& OutInfo) const override;
 #endif // WITH_EDITORONLY_DATA
 
diff --git a/Engine/Plugins/Runtime/GeometryCache/Source/GeometryCache/Private/NiagaraGeometryCacheRendererProperties.cpp b/Engine/Plugins/Runtime/GeometryCache/Source/GeometryCache/Private/NiagaraGeometryCacheRendererProperties.cpp
--- a/Engine/Plugins/Runtime/GeometryCache/Source/GeometryCache/Private/NiagaraGeometryCacheRendererProperties.cpp
+++ b/Engine/Plugins/Runtime/GeometryCache/Source/GeometryCache/Private/NiagaraGeometryCacheRendererProperties.cpp
@@ -333,23 +333,6 @@ const TArray<FNiagaraVariable>& UNiagaraGeometryCacheRendererProperties::GetOpti
 	return Attrs;
 }
 
-bool UNiagaraGeometryCacheRendererProperties::IsMaterialValidForRenderer(UMaterial* Material, FText& InvalidMessage)
-{
-	if (Material->bUsedWithGeometryCache == false)
-	{
-		InvalidMessage = NSLOCTEXT("NiagaraGeometryCacheRendererProperties", "InvalidMaterialMessage", "The material isn't marked as \"Used with Geometry Cache\"");
-		return false;
-	}
-	return true;
-}
-
-void UNiagaraGeometryCacheRendererProperties::FixMaterial(UMaterial* Material)
-{
-	Material->Modify();
-	Material->bUsedWithGeometryCache = true;
-	Material->ForceRecompileForRendering();
-}
-
 void UNiagaraGeometryCacheRendererProperties::GetRendererFeedback(const UNiagaraEmitter* InEmitter,	TArray<FText>& OutErrors, TArray<FText>& OutWarnings, TArray<FText>& OutInfo) const
 {
 	OutInfo.Add(FText::FromString(TEXT("The geometry cache renderer is still an experimental feature.")));
