Sequencer: Expose FMovieSceneGeometryCacheParams as a BlueprintType

diff --git a/Engine/Plugins/Runtime/GeometryCache/Source/GeometryCacheTracks/Classes/MovieSceneGeometryCacheSection.h b/Engine/Plugins/Runtime/GeometryCache/Source/GeometryCacheTracks/Classes/MovieSceneGeometryCacheSection.h
--- a/Engine/Plugins/Runtime/GeometryCache/Source/GeometryCacheTracks/Classes/MovieSceneGeometryCacheSection.h
+++ b/Engine/Plugins/Runtime/GeometryCache/Source/GeometryCacheTracks/Classes/MovieSceneGeometryCacheSection.h
@@ -11,7 +11,7 @@
 #include "UObject/SoftObjectPath.h"
 #include "MovieSceneGeometryCacheSection.generated.h"
 
-USTRUCT()
+USTRUCT(BlueprintType)
  struct GEOMETRYCACHETRACKS_API FMovieSceneGeometryCacheParams
 {
 	GENERATED_BODY()
