Fixed a bug where array index changed were not picked up correctly for the geometry cache renderer.

diff --git a/Engine/Plugins/Runtime/GeometryCache/Source/GeometryCache/Private/NiagaraRendererGeometryCache.cpp b/Engine/Plugins/Runtime/GeometryCache/Source/GeometryCache/Private/NiagaraRendererGeometryCache.cpp
--- a/Engine/Plugins/Runtime/GeometryCache/Source/GeometryCache/Private/NiagaraRendererGeometryCache.cpp
+++ b/Engine/Plugins/Runtime/GeometryCache/Source/GeometryCache/Private/NiagaraRendererGeometryCache.cpp
@@ -81,6 +81,26 @@ UGeometryCache* ResolveGeometryCache(const FNiagaraGeometryCacheReference& Entry
 	return FoundCache;
 }
 
+int32 GetGeometryCacheIndex(const FNiagaraDataSetReaderInt32<int32>& ArrayIndexAccessor, bool bCreateRandomIfUnassigned, const FNiagaraEmitterInstance* Emitter, const UNiagaraGeometryCacheRendererProperties* Properties, int32 ParticleIndex, int32 ParticleID)
+{
+	if (bCreateRandomIfUnassigned == false && ArrayIndexAccessor.IsValid() == false)
+	{
+		return INDEX_NONE;
+	}
+	int32 CacheIndex = ArrayIndexAccessor.GetSafe(ParticleIndex, INDEX_NONE);
+	if (CacheIndex == INDEX_NONE && Emitter->GetCachedEmitterData()->bDeterminism)
+	{
+		int32 Seed = Properties->bAssignComponentsOnParticleID ? ParticleID : ParticleIndex;
+		FRandomStream RandomStream = FRandomStream(Seed * 907633515U); // multiply the seed, otherwise we get very poor randomness for small seeds
+		CacheIndex = RandomStream.RandRange(0, Properties->GeometryCaches.Num());
+	}
+	else if (CacheIndex == INDEX_NONE)
+	{
+		CacheIndex = FMath::RandRange(0, Properties->GeometryCaches.Num() - 1);
+	}
+	return CacheIndex % Properties->GeometryCaches.Num(); // wrap around if index is too large
+}
+
 /** Update render data buffer from attributes */
 void FNiagaraRendererGeometryCache::PostSystemTick_GameThread(const UNiagaraRendererProperties* InProperties, const FNiagaraEmitterInstance* Emitter)
 {
@@ -231,20 +251,10 @@ void FNiagaraRendererGeometryCache::PostSystemTick_GameThread(const UNiagaraRend
 			GeometryComponent = ComponentPool[PoolIndex].Component.Get();
 		}
 
-		if (!GeometryComponent || GeometryComponent->HasAnyFlags(RF_BeginDestroyed | RF_FinishDestroyed))
+		bool bCreateNewComponent = !GeometryComponent || GeometryComponent->HasAnyFlags(RF_BeginDestroyed | RF_FinishDestroyed);
+		int32 CacheIndex = GetGeometryCacheIndex(ArrayIndexAccessor, bCreateNewComponent, Emitter, Properties, ParticleIndex, ParticleID);
+		if (Properties->GeometryCaches.IsValidIndex(CacheIndex))
 		{
-			int32 CacheIndex = ArrayIndexAccessor.GetSafe(ParticleIndex, INDEX_NONE);
-			if (CacheIndex == INDEX_NONE && Emitter->GetCachedEmitterData()->bDeterminism)
-			{
-				int32 Seed = Properties->bAssignComponentsOnParticleID ? ParticleID : ParticleIndex;
-				FRandomStream RandomStream = FRandomStream(Seed * 907633515U); // multiply the seed, otherwise we get very poor randomness for small seeds
-				CacheIndex = RandomStream.RandRange(0, Properties->GeometryCaches.Num());
-			}
-			else if (CacheIndex == INDEX_NONE)
-			{
-				CacheIndex = FMath::RandRange(0, Properties->GeometryCaches.Num() - 1);
-			}
-			CacheIndex = CacheIndex % Properties->GeometryCaches.Num(); // wrap around if index is too large
 			const FNiagaraGeometryCacheReference& Entry = Properties->GeometryCaches[CacheIndex];
 			UGeometryCache* Cache = ResolveGeometryCache(Entry, Emitter);
 			if (!Cache)
@@ -252,61 +262,70 @@ void FNiagaraRendererGeometryCache::PostSystemTick_GameThread(const UNiagaraRend
 				// no valid geometry cache was provided for this particle
 				continue;
 			}
-
-			// Determine the owner actor or spawn one
-			AActor* OwnerActor = SpawnedOwner.Get();
-			if (OwnerActor == nullptr)
+			
+			if (bCreateNewComponent)
 			{
-				OwnerActor = AttachComponent->GetOwner();
+				// Determine the owner actor or spawn one
+				AActor* OwnerActor = SpawnedOwner.Get();
 				if (OwnerActor == nullptr)
 				{
-					// NOTE: This can happen with spawned systems
-					OwnerActor = AttachComponent->GetWorld()->SpawnActor<AActor>();
-					OwnerActor->SetFlags(RF_Transient);
-					SpawnedOwner = OwnerActor;
+					OwnerActor = AttachComponent->GetOwner();
+					if (OwnerActor == nullptr)
+					{
+						// NOTE: This can happen with spawned systems
+						OwnerActor = AttachComponent->GetWorld()->SpawnActor<AActor>();
+						OwnerActor->SetFlags(RF_Transient);
+						SpawnedOwner = OwnerActor;
+					}
 				}
-			}
 
-			// if we don't have a pooled component we create a new one from the template
-			GeometryComponent = NewObject<UGeometryCacheComponent>(OwnerActor);
-			IGeometryCacheStreamingManager::Get().AddStreamingComponent(GeometryComponent);
-			GeometryComponent->SetFlags(RF_Transient);
-			GeometryComponent->SetupAttachment(AttachComponent);
-			GeometryComponent->RegisterComponent();
-			GeometryComponent->AddTickPrerequisiteComponent(AttachComponent);
-			GeometryComponent->SetLooping(Properties->bIsLooping);
-			GeometryComponent->SetGeometryCache(Cache);
-			GeometryComponent->SetManualTick(true); // we want to tick the component with the delta time of the niagara sim
+				// if we don't have a pooled component we create a new one from the template
+				GeometryComponent = NewObject<UGeometryCacheComponent>(OwnerActor);
+				IGeometryCacheStreamingManager::Get().AddStreamingComponent(GeometryComponent);
+				GeometryComponent->SetFlags(RF_Transient);
+				GeometryComponent->SetupAttachment(AttachComponent);
+				GeometryComponent->RegisterComponent();
+				GeometryComponent->AddTickPrerequisiteComponent(AttachComponent);
+				GeometryComponent->SetLooping(Properties->bIsLooping);
+				GeometryComponent->SetManualTick(true); // we want to tick the component with the delta time of the niagara sim
+
+				if (Emitter->GetCachedEmitterData()->bLocalSpace)
+				{
+					GeometryComponent->SetAbsolute(false, false, false);
+				}
+				else
+				{
+					GeometryComponent->SetAbsolute(true, true, true);
+				}
 
-			if (Emitter->GetCachedEmitterData()->bLocalSpace)
-			{
-				GeometryComponent->SetAbsolute(false, false, false);
-			}
-			else
-			{
-				GeometryComponent->SetAbsolute(true, true, true);
+				if (PoolIndex >= 0)
+				{
+					// This should only happen if the component was destroyed externally
+					ComponentPool[PoolIndex].Component = GeometryComponent;
+					ComponentPool[PoolIndex].LastElapsedTime = 0;
+				}
+				else
+				{
+					// Add a new pool entry
+					PoolIndex = ComponentPool.Num();
+					ComponentPool.AddDefaulted_GetRef().Component = GeometryComponent;
+					ComponentPool[PoolIndex].LastElapsedTime = 0;
+				}
 			}
 
-			// set the override materials
+			// set the cache and override materials, since they might have changed from the last time we accessed the array index
+			// if they did not change from last tick these calls should be no-ops.
+			GeometryComponent->SetGeometryCache(Cache);
 			for (int32 i = 0; i < Entry.OverrideMaterials.Num(); i++)
 			{
 				const TObjectPtr<UMaterialInterface>& OverrideMat = Entry.OverrideMaterials[i];
 				GeometryComponent->SetMaterial(i, OverrideMat);
 			}
-
-			if (PoolIndex >= 0)
-			{
-				// This should only happen if the component was destroyed externally
-				ComponentPool[PoolIndex].Component = GeometryComponent;
-				ComponentPool[PoolIndex].LastElapsedTime = 0;
-			}
-			else
-			{
-				// Add a new pool entry
-				PoolIndex = ComponentPool.Num();
-				ComponentPool.AddDefaulted_GetRef().Component = GeometryComponent;
-				ComponentPool[PoolIndex].LastElapsedTime = 0;
-			}
+		}
+		else if (bCreateNewComponent)
+		{
+			// can't create a new component without valid geometry cache
+			continue;
 		}
 
 		const auto PositionAccessor = FNiagaraDataSetAccessor<FNiagaraPosition>::CreateReader(Data, Properties->PositionBinding.GetDataSetBindableVariable().GetName());
