Fix TotalPrimitiveCount calculation when index buffer validation fails (when index buffer contains fewer triangles than specified in FGeometryCacheMeshBatchInfo::NumTriangles)

diff --git a/Engine/Plugins/Runtime/GeometryCache/Source/GeometryCache/Private/GeometryCacheSceneProxy.cpp b/Engine/Plugins/Runtime/GeometryCache/Source/GeometryCache/Private/GeometryCacheSceneProxy.cpp
--- a/Engine/Plugins/Runtime/GeometryCache/Source/GeometryCache/Private/GeometryCacheSceneProxy.cpp
+++ b/Engine/Plugins/Runtime/GeometryCache/Source/GeometryCache/Private/GeometryCacheSceneProxy.cpp
@@ -712,7 +712,7 @@ void FGeometryCacheSceneProxy::UpdateAnimation(float NewTime, bool bNewLooping,
 					Segment.MaxVertices = Section->PositionBuffers[PositionBufferIndex].GetSizeInBytes() / Segment.VertexBufferStride; // conservative estimate
 
 					Segments.Add(Segment);
-					TotalPrimitiveCount += BatchInfo.NumTriangles;
+					TotalPrimitiveCount += Segment.NumPrimitives;
 				}
 
 				bRequireRecreate = bRequireRecreate || Section->RayTracingGeometry.Initializer.IndexBuffer != Section->IndexBuffer.IndexBufferRHI;
