Fixed GeometryCacheSceneProxy would sometimes select an invalid frame's mesh data because it failed to decode the frame but the buffer still had data from a previous frame.

diff --git a/Engine/Plugins/Runtime/GeometryCache/Source/GeometryCache/Private/GeometryCacheSceneProxy.cpp b/Engine/Plugins/Runtime/GeometryCache/Source/GeometryCache/Private/GeometryCacheSceneProxy.cpp
--- a/Engine/Plugins/Runtime/GeometryCache/Source/GeometryCache/Private/GeometryCacheSceneProxy.cpp
+++ b/Engine/Plugins/Runtime/GeometryCache/Source/GeometryCache/Private/GeometryCacheSceneProxy.cpp
@@ -725,8 +725,8 @@ void FGeometryCacheSceneProxy::UpdateAnimation(float NewTime, bool bNewLooping,
 				{
 					if (bRequireRecreate)
 					{
-						Section->RayTracingGeometry.UpdateRHI();
-					}
+					Section->RayTracingGeometry.UpdateRHI();
+				}
 					else
 					{
 						// Request full build on same geometry because data might have changed to much for update call?
@@ -1175,7 +1175,7 @@ void FGeometryCacheSceneProxy::FrameUpdate() const
 				// Only bother uploading if anything changed or when the we failed to decode anything make sure update the gpu buffers regardless
 				if (bFrameIndicesChanged || bDifferentRoundedInterpolationFactor || (bDifferentInterpolationFactor && bExtrapolateFrames) || bDecodedAnything || bDecoderError)
 				{
-					const bool bNextFrame = !!FMath::RoundToInt(InterpolationFactor) && TrackProxy->NextFrameMeshData->Positions.Num() > 0; // use next frame only if it's valid
+					const bool bNextFrame = !!FMath::RoundToInt(InterpolationFactor) && TrackProxy->NextFrameMeshData->Positions.Num() > 0 && (TrackProxy->NextFrameIndex != -1); // use next frame only if it's valid
 					const uint32 FrameIndexToUse = bNextFrame ? TrackProxy->NextFrameIndex : TrackProxy->FrameIndex;
 					FGeometryCacheMeshData* MeshDataToUse = bNextFrame ? TrackProxy->NextFrameMeshData : TrackProxy->MeshData;
 
