Alembic: Fixed crash when importing GeometryCache with "Flatten Tracks" disabled from Alembic with entirely invisible objects.

diff --git a/Engine/Plugins/Importers/AlembicImporter/Source/AlembicLibrary/Private/AbcImporter.cpp b/Engine/Plugins/Importers/AlembicImporter/Source/AlembicLibrary/Private/AbcImporter.cpp
--- a/Engine/Plugins/Importers/AlembicImporter/Source/AlembicLibrary/Private/AbcImporter.cpp
+++ b/Engine/Plugins/Importers/AlembicImporter/Source/AlembicLibrary/Private/AbcImporter.cpp
@@ -488,7 +488,7 @@ UGeometryCache* FAbcImporter::ImportAsGeometryCache(UObject* InParent, EObjectFl
 
 							// Generate the mesh data for this sample
 							const bool bVisible = PolyMesh->GetVisibility(FrameIndex);
-							const float FrameTime = PolyMesh->GetTimeForFrameIndex(FrameIndex)- InAbcFile->GetImportTimeOffset();
+							const float FrameTime = PolyMesh->GetTimeForFrameIndex(FrameIndex) - InAbcFile->GetImportTimeOffset();
 							if (bVisible)
 							{
 								const bool bUseVelocitiesAsMotionVectors = ( ImportSettings->GeometryCacheSettings.MotionVectors == EAbcGeometryCacheMotionVectorsImport::ImportAbcVelocitiesAsMotionVectors );
@@ -518,8 +518,11 @@ UGeometryCache* FAbcImporter::ImportAsGeometryCache(UObject* InParent, EObjectFl
 				MatTimes.Add(AbcFile->GetImportLength() + AbcFile->GetImportTimeOffset());
 				Track->SetMatrixSamples(Mats, MatTimes);
 
-				Track->EndCoding();
-				GeometryCache->AddTrack(Track);
+				// Some tracks might not have any mesh samples because they are invisible (can happen with bFlattenTracks=false), so skip them
+				if (Track->EndCoding())
+				{
+					GeometryCache->AddTrack(Track);
+				}
 			}
 		}
 
diff --git a/Engine/Plugins/Runtime/GeometryCache/Source/GeometryCache/Classes/GeometryCacheTrackStreamable.h b/Engine/Plugins/Runtime/GeometryCache/Source/GeometryCache/Classes/GeometryCacheTrackStreamable.h
--- a/Engine/Plugins/Runtime/GeometryCache/Source/GeometryCache/Classes/GeometryCacheTrackStreamable.h
+++ b/Engine/Plugins/Runtime/GeometryCache/Source/GeometryCache/Classes/GeometryCacheTrackStreamable.h
@@ -137,9 +137,9 @@ class GEOMETRYCACHE_API UGeometryCacheTrackStreamable : public UGeometryCacheTra
 	void AddVisibilitySample(const bool bVisible, const float SampleTime);
 
 	/**
-		Finish up coding
+		Finish up coding. Return true if the track has samples
 	*/
-	void EndCoding();
+	bool EndCoding();
 #endif
 
 	/** Codec for this track */
diff --git a/Engine/Plugins/Runtime/GeometryCache/Source/GeometryCache/Private/GeometryCacheTrackStreamable.cpp b/Engine/Plugins/Runtime/GeometryCache/Source/GeometryCache/Private/GeometryCacheTrackStreamable.cpp
--- a/Engine/Plugins/Runtime/GeometryCache/Source/GeometryCache/Private/GeometryCacheTrackStreamable.cpp
+++ b/Engine/Plugins/Runtime/GeometryCache/Source/GeometryCache/Private/GeometryCacheTrackStreamable.cpp
@@ -182,7 +182,7 @@ void UGeometryCacheTrackStreamable::BeginCoding(UGeometryCacheCodecBase* InCodec
 	VisibilitySamples.Empty();
 }
 
-void UGeometryCacheTrackStreamable::EndCoding()
+bool UGeometryCacheTrackStreamable::EndCoding()
 {
 	// This needs to be deleted first so it flushes any buffered frames before
 	// we call EndCoding on the codec
@@ -192,7 +192,10 @@ void UGeometryCacheTrackStreamable::EndCoding()
 	Codec->EndCoding();
 	InitializeRenderResources();
 
-	StartSampleTime = FMath::Min(Samples[0].SampleTime, 0.0f);
+	if (Samples.Num() > 0)
+	{
+		StartSampleTime = FMath::Min(Samples[0].SampleTime, 0.0f);
+	}
 
 	if (ImportVisibilitySamples.Num())
 	{
@@ -229,7 +232,7 @@ void UGeometryCacheTrackStreamable::EndCoding()
 	}
 	else
 	{
-		TRange<float> VisibilityRange(StartSampleTime, Samples.Last().SampleTime);
+		TRange<float> VisibilityRange(StartSampleTime, Samples.Num() > 0 ? Samples.Last().SampleTime : 0.0f);
 		FVisibilitySample Sample(true);
 		Sample.Range = VisibilityRange;
 
@@ -241,6 +244,8 @@ void UGeometryCacheTrackStreamable::EndCoding()
 	{
 		Duration = Samples.Last().SampleTime - Samples[0].SampleTime;
 	}
+
+	return Samples.Num() > 0;
 }
 
 void UGeometryCacheTrackStreamable::AddMeshSample(const FGeometryCacheMeshData& MeshData, const float SampleTime, bool bSameTopologyAsPrevious)
