Added LLM tags for GeometryCache and Alembic allocations.

diff --git a/Engine/Plugins/Importers/AlembicImporter/Source/AlembicLibrary/Private/AlembicLibraryModule.cpp b/Engine/Plugins/Importers/AlembicImporter/Source/AlembicLibrary/Private/AlembicLibraryModule.cpp
--- a/Engine/Plugins/Importers/AlembicImporter/Source/AlembicLibrary/Private/AlembicLibraryModule.cpp
+++ b/Engine/Plugins/Importers/AlembicImporter/Source/AlembicLibrary/Private/AlembicLibraryModule.cpp
@@ -2,15 +2,18 @@
 
 #include "AlembicLibraryModule.h"
 #include "AbcImporter.h"
-
 #include "AbcImportSettingsCustomization.h"
-
+#include "HAL/LowLevelMemTracker.h"
 #include "PropertyEditorModule.h"
 
+LLM_DEFINE_TAG(Alembic);
+
 class FAlembicLibraryModule : public IAlembicLibraryModule
 {
 	virtual void StartupModule() override
 	{
+		LLM_SCOPE_BYTAG(Alembic);
+
 		FModuleManager::LoadModuleChecked<IModuleInterface>("GeometryCache");
 
 		// Register class/struct customizations
diff --git a/Engine/Plugins/Runtime/GeometryCache/Source/GeometryCache/Private/GeometryCacheModule.cpp b/Engine/Plugins/Runtime/GeometryCache/Source/GeometryCache/Private/GeometryCacheModule.cpp
--- a/Engine/Plugins/Runtime/GeometryCache/Source/GeometryCache/Private/GeometryCacheModule.cpp
+++ b/Engine/Plugins/Runtime/GeometryCache/Source/GeometryCache/Private/GeometryCacheModule.cpp
@@ -1,17 +1,21 @@
 // Copyright Epic Games, Inc. All Rights Reserved.
 
 #include "GeometryCacheModule.h"
+#include "CodecV1.h"
 #if WITH_EDITOR
 #include "GeometryCacheEdModule.h"
 #endif // WITH_EDITOR
 #include "GeometryCacheStreamingManager.h"
-#include "CodecV1.h"
 #include "NiagaraGeometryCacheRendererProperties.h"
 
 IMPLEMENT_MODULE(FGeometryCacheModule, GeometryCache)
 
+LLM_DEFINE_TAG(GeometryCache);
+
 void FGeometryCacheModule::StartupModule()
 {
+	LLM_SCOPE_BYTAG(GeometryCache);
+
 #if WITH_EDITOR
 	FGeometryCacheEdModule& Module = FModuleManager::LoadModuleChecked<FGeometryCacheEdModule>(TEXT("GeometryCacheEd"));
 #endif
diff --git a/Engine/Plugins/Runtime/GeometryCache/Source/GeometryCache/Public/GeometryCacheModule.h b/Engine/Plugins/Runtime/GeometryCache/Source/GeometryCache/Public/GeometryCacheModule.h
--- a/Engine/Plugins/Runtime/GeometryCache/Source/GeometryCache/Public/GeometryCacheModule.h
+++ b/Engine/Plugins/Runtime/GeometryCache/Source/GeometryCache/Public/GeometryCacheModule.h
@@ -3,10 +3,13 @@
 #pragma once
 
 #include "CoreMinimal.h"
+#include "HAL/LowLevelMemTracker.h"
 #include "Modules/ModuleInterface.h"
 #include "Modules/ModuleManager.h"
 #include "Stats/Stats.h"
 
+LLM_DECLARE_TAG_API(GeometryCache, GEOMETRYCACHE_API);
+
 /**
  * The public interface to this module
  */
diff --git a/Engine/Plugins/Runtime/GeometryCache/Source/GeometryCacheEd/Private/GeometryCacheEdModule.cpp b/Engine/Plugins/Runtime/GeometryCache/Source/GeometryCacheEd/Private/GeometryCacheEdModule.cpp
--- a/Engine/Plugins/Runtime/GeometryCache/Source/GeometryCacheEd/Private/GeometryCacheEdModule.cpp
+++ b/Engine/Plugins/Runtime/GeometryCache/Source/GeometryCacheEd/Private/GeometryCacheEdModule.cpp
@@ -8,12 +8,15 @@
 #include "GeometryCacheAssetBroker.h"
 #include "GeometryCache.h"
 #include "GeometryCacheComponent.h"
+#include "GeometryCacheModule.h"
 #include "GeometryCacheThumbnailRenderer.h"
 
 IMPLEMENT_MODULE(FGeometryCacheEdModule, GeometryCacheEd)
 
 void FGeometryCacheEdModule::StartupModule()
 {
+	LLM_SCOPE_BYTAG(GeometryCache);
+
 	FAssetToolsModule& AssetToolsModule = FAssetToolsModule::GetModule();
 
 	IAssetTools& AssetTools = AssetToolsModule.Get();
diff --git a/Engine/Plugins/Runtime/GeometryCache/Source/GeometryCacheSequencer/Public/GeometryCacheSequencerModule.h b/Engine/Plugins/Runtime/GeometryCache/Source/GeometryCacheSequencer/Public/GeometryCacheSequencerModule.h
--- a/Engine/Plugins/Runtime/GeometryCache/Source/GeometryCacheSequencer/Public/GeometryCacheSequencerModule.h
+++ b/Engine/Plugins/Runtime/GeometryCache/Source/GeometryCacheSequencer/Public/GeometryCacheSequencerModule.h
@@ -1,11 +1,13 @@
 // Copyright Epic Games, Inc. All Rights Reserved.
 
 #pragma once
-#include "ISequencerModule.h"
+
 #include "CoreMinimal.h"
+#include "GeometryCacheModule.h"
+#include "GeometryCacheTrackEditor.h"
+#include "ISequencerModule.h"
 #include "Modules/ModuleInterface.h"
 #include "Modules/ModuleManager.h"
-#include "GeometryCacheTrackEditor.h"
 
 /**
  * The public interface to this module
@@ -16,6 +18,7 @@ public:
 
 	virtual void StartupModule() override
 	{
+		LLM_SCOPE_BYTAG(GeometryCache);
 
 		ISequencerModule& SequencerModule = FModuleManager::Get().LoadModuleChecked<ISequencerModule>("Sequencer");
 		TrackEditorBindingHandle = SequencerModule.RegisterTrackEditor(FOnCreateTrackEditor::CreateStatic(&FGeometryCacheTrackEditor::CreateTrackEditor));
diff --git a/Engine/Plugins/Runtime/GeometryCache/Source/GeometryCacheTracks/Private/GeometryCacheTracksModule.cpp b/Engine/Plugins/Runtime/GeometryCache/Source/GeometryCacheTracks/Private/GeometryCacheTracksModule.cpp
--- a/Engine/Plugins/Runtime/GeometryCache/Source/GeometryCacheTracks/Private/GeometryCacheTracksModule.cpp
+++ b/Engine/Plugins/Runtime/GeometryCache/Source/GeometryCacheTracks/Private/GeometryCacheTracksModule.cpp
@@ -1,6 +1,7 @@
 // Copyright Epic Games, Inc. All Rights Reserved.
 
 #include "GeometryCacheTracksModule.h"
+#include "GeometryCacheModule.h"
 #if WITH_EDITOR
 #include "GeometryCacheSequencerModule.h"
 #endif // WITH_EDITOR
@@ -9,6 +10,7 @@ IMPLEMENT_MODULE(FGeometryCacheTracksModule, GeometryCacheTracks)
 
 void FGeometryCacheTracksModule::StartupModule()
 {
+	LLM_SCOPE_BYTAG(GeometryCache);
 #if WITH_EDITOR
 	FGeometryCacheSequencerModule& Module = FModuleManager::LoadModuleChecked<FGeometryCacheSequencerModule>(TEXT("GeometryCacheSequencer"));
 #endif
