Geometry Cache Track:
- Export Geometry Cache Track Editor (done similar to other Track Editors such as FMediaTrackEditor)

diff --git a/Engine/Plugins/Runtime/GeometryCache/Source/GeometryCacheSequencer/Classes/GeometryCacheTrackEditor.h b/Engine/Plugins/Runtime/GeometryCache/Source/GeometryCacheSequencer/Classes/GeometryCacheTrackEditor.h
--- a/Engine/Plugins/Runtime/GeometryCache/Source/GeometryCacheSequencer/Classes/GeometryCacheTrackEditor.h
+++ b/Engine/Plugins/Runtime/GeometryCache/Source/GeometryCacheSequencer/Classes/GeometryCacheTrackEditor.h
@@ -22,7 +22,7 @@ class UGeometryCacheComponent;
 /**
  * Tools for Geometry Cache tracks
  */
-class FGeometryCacheTrackEditor : public FMovieSceneTrackEditor
+class GEOMETRYCACHESEQUENCER_API FGeometryCacheTrackEditor : public FMovieSceneTrackEditor
 {
 public:
 	/**
