Improved geometry cache default asset

diff --git a/Engine/Plugins/Runtime/GeometryCache/Source/GeometryCacheEd/Private/GeometryCacheEdModule.cpp b/Engine/Plugins/Runtime/GeometryCache/Source/GeometryCacheEd/Private/GeometryCacheEdModule.cpp
--- a/Engine/Plugins/Runtime/GeometryCache/Source/GeometryCacheEd/Private/GeometryCacheEdModule.cpp
+++ b/Engine/Plugins/Runtime/GeometryCache/Source/GeometryCacheEd/Private/GeometryCacheEdModule.cpp
@@ -40,7 +40,7 @@ void FGeometryCacheEdModule::StartupModule()
 			UNiagaraGeometryCacheRendererProperties* NewRenderer = NewObject<UNiagaraGeometryCacheRendererProperties>(OuterEmitter, NAME_None, RF_Transactional);
 			if(ensure(NewRenderer->GeometryCaches.Num() == 1))
 			{
-				FSoftObjectPath DefaultGeometryCache(TEXT("GeometryCache'/Niagara/DefaultAssets/DefaultGeometryCacheAsset.DefaultGeometryCacheAsset'"));
+				FSoftObjectPath DefaultGeometryCache(TEXT("GeometryCache'/Niagara/DefaultAssets/GeometryCache/DefaultGeometryCacheAsset.DefaultGeometryCacheAsset'"));
 				NewRenderer->GeometryCaches[0].GeometryCache = Cast<UGeometryCache>(DefaultGeometryCache.TryLoad());
 			}
 			return NewRenderer;
