USD: GeometryCache from USD now use tracks in LevelSequence instead of the Time property to drive their animations. Added USD.DisableGeoCacheTracks cvar to revert to the previous behavior.

diff --git a/Engine/Plugins/Importers/USDImporter/Source/USDStage/USDStage.Build.cs b/Engine/Plugins/Importers/USDImporter/Source/USDStage/USDStage.Build.cs
--- a/Engine/Plugins/Importers/USDImporter/Source/USDStage/USDStage.Build.cs
+++ b/Engine/Plugins/Importers/USDImporter/Source/USDStage/USDStage.Build.cs
@@ -25,6 +25,9 @@ namespace UnrealBuildTool.Rules
 					"Core",
 					"CoreUObject",
 					"Engine",
+					"GeometryCache",
+					"GeometryCacheTracks",
+					"GeometryCacheUSD",
 					"HairStrandsCore",
 					"LevelSequence",
 					"LiveLinkComponents", // For tracking edits to LiveLinkComponentController properties and writing to USD
diff --git a/Engine/Plugins/Runtime/GeometryCache/Source/GeometryCache/Private/GeometryCacheSceneProxy.cpp b/Engine/Plugins/Runtime/GeometryCache/Source/GeometryCache/Private/GeometryCacheSceneProxy.cpp
--- a/Engine/Plugins/Runtime/GeometryCache/Source/GeometryCache/Private/GeometryCacheSceneProxy.cpp
+++ b/Engine/Plugins/Runtime/GeometryCache/Source/GeometryCache/Private/GeometryCacheSceneProxy.cpp
@@ -745,6 +745,8 @@ void FGeometryCacheSceneProxy::UpdateAnimation(float NewTime, bool bNewLooping,
 
 void FGeometryCacheSceneProxy::FrameUpdate() const
 {
+	TRACE_CPUPROFILER_EVENT_SCOPE(FGeometryCacheSceneProxy::FrameUpdate);
+
 	for (FGeomCacheTrackProxy* TrackProxy : Tracks)
 	{
 		// Render out stored TrackProxy's
@@ -849,6 +851,8 @@ void FGeometryCacheSceneProxy::FrameUpdate() const
 			// Can we interpolate the vertex data?
 			if (bCanInterpolate && (bDifferentInterpolationFactor || bFrameIndicesChanged) && !bDecoderError && CVarInterpolateFrames.GetValueOnRenderThread() != 0)
 			{
+				TRACE_CPUPROFILER_EVENT_SCOPE(FGeometryCacheSceneProxy::FrameUpdate::Interpolate);
+
 				TrackProxy->bNextFrameMeshDataSelected = false;
 
 				SCOPE_CYCLE_COUNTER(STAT_InterpolateFrames);
@@ -1171,6 +1175,7 @@ void FGeometryCacheSceneProxy::FrameUpdate() const
 			else
 			{
 				// We just don't interpolate between frames if we got GPU to burn we could someday render twice and stipple fade between it :-D like with lods
+				TRACE_CPUPROFILER_EVENT_SCOPE(FGeometryCacheSceneProxy::FrameUpdate::NonInterpolate);
 
 				// Only bother uploading if anything changed or when the we failed to decode anything make sure update the gpu buffers regardless
 				if (bFrameIndicesChanged || bDifferentRoundedInterpolationFactor || (bDifferentInterpolationFactor && bExtrapolateFrames) || bDecodedAnything || bDecoderError)
