GeometryCacheStreamingManager fix to not deadlock MapChunk

diff --git a/Engine/Plugins/Runtime/GeometryCache/Source/GeometryCache/Private/GeometryCacheStreamingManager.cpp b/Engine/Plugins/Runtime/GeometryCache/Source/GeometryCache/Private/GeometryCacheStreamingManager.cpp
--- a/Engine/Plugins/Runtime/GeometryCache/Source/GeometryCache/Private/GeometryCacheStreamingManager.cpp
+++ b/Engine/Plugins/Runtime/GeometryCache/Source/GeometryCache/Private/GeometryCacheStreamingManager.cpp
@@ -1,13 +1,14 @@
 // Copyright Epic Games, Inc. All Rights Reserved.
 
 #include "GeometryCacheStreamingManager.h"
-#include "GeometryCacheComponent.h"
+
 #include "GeometryCache.h"
-#include "GeometryCacheTrackStreamable.h"
+#include "GeometryCacheComponent.h"
 #include "GeometryCacheModule.h"
-#include "Misc/ScopeRWLock.h"
+#include "GeometryCacheTrackStreamable.h"
 #include "StreamingGeometryCacheData.h"
 #include "Async/Async.h"
+#include "Misc/ScopeRWLock.h"
 
 DECLARE_CYCLE_STAT(TEXT("Update Resource Streaming"), STAT_UpdateResourceStreaming, STATGROUP_GeometryCache);
 DECLARE_CYCLE_STAT(TEXT("Wait Untill Requests Finished"), STAT_BlockTillAllRequestsFinished, STATGROUP_GeometryCache);
@@ -226,13 +227,37 @@ void FGeometryCacheStreamingManager::RemoveGeometryCache(UGeometryCacheTrackStre
 {
 	check(IsInGameThread());
 
-	FRWScopeLock WriteLock(StreamingGeometryCachesLock, SLT_Write);
-	FStreamingGeometryCacheData** CacheData = StreamingGeometryCaches.Find(Cache);
+	// Reorganized to not hold the lock while calling FlushRenderingCommands();
+	// This causes a deadlock with MapChunk.
+
+	bool bCacheFound = false;
+	{
+		FRWScopeLock ReadLock(StreamingGeometryCachesLock, SLT_ReadOnly);
+		bCacheFound = StreamingGeometryCaches.Contains(Cache);
+	}
+
+	// FlushRenderingCommands first, without holding the lock.
+	if (bCacheFound)
+	{
+		// Flush the render thread so any decoding still happening is finished and thus no maps held by the render thread.
+		FlushRenderingCommands();
+	}
+
+	FStreamingGeometryCacheData* CacheData = nullptr;
+	{
+		FRWScopeLock WriteLock(StreamingGeometryCachesLock, SLT_Write);
+		FStreamingGeometryCacheData** FoundCacheData = StreamingGeometryCaches.Find(Cache);
+		if (FoundCacheData)
+		{
+			CacheData = *FoundCacheData;
+			StreamingGeometryCaches.Remove(Cache);
+		}
+	}
+
 	if (CacheData != nullptr)
 	{
-		check(*CacheData != nullptr);
-		delete (*CacheData);
-		StreamingGeometryCaches.Remove(Cache);
+		check(CacheData != nullptr);
+		delete CacheData;	// This will FlushRenderingCommands again so we move it outside the lock.
 	}
 }
 
