Fixed a bug where component renderers were still visible when another emitter was running isolated

diff --git a/Engine/Plugins/FX/Niagara/Source/Niagara/Private/NiagaraRendererComponents.cpp b/Engine/Plugins/FX/Niagara/Source/Niagara/Private/NiagaraRendererComponents.cpp
--- a/Engine/Plugins/FX/Niagara/Source/Niagara/Private/NiagaraRendererComponents.cpp
+++ b/Engine/Plugins/FX/Niagara/Source/Niagara/Private/NiagaraRendererComponents.cpp
@@ -385,6 +385,14 @@ void FNiagaraRendererComponents::PostSystemTick_GameThread(const UNiagaraRendere
 		return;
 	}
 
+#if WITH_EDITORONLY_DATA
+	if (SystemInstance->GetIsolateEnabled() && !Emitter->GetEmitterHandle().IsIsolated())
+	{
+		ResetComponentPool(true);
+		return;
+	}
+#endif
+
 	USceneComponent* AttachComponent = SystemInstance->GetAttachComponent();
 	if (!AttachComponent)
 	{
diff --git a/Engine/Plugins/Runtime/GeometryCache/Source/GeometryCache/Private/NiagaraRendererGeometryCache.cpp b/Engine/Plugins/Runtime/GeometryCache/Source/GeometryCache/Private/NiagaraRendererGeometryCache.cpp
--- a/Engine/Plugins/Runtime/GeometryCache/Source/GeometryCache/Private/NiagaraRendererGeometryCache.cpp
+++ b/Engine/Plugins/Runtime/GeometryCache/Source/GeometryCache/Private/NiagaraRendererGeometryCache.cpp
@@ -114,6 +114,14 @@ void FNiagaraRendererGeometryCache::PostSystemTick_GameThread(const UNiagaraRend
 		return;
 	}
 
+#if WITH_EDITORONLY_DATA
+	if (SystemInstance->GetIsolateEnabled() && !Emitter->GetEmitterHandle().IsIsolated())
+	{
+		ResetComponentPool(true);
+		return;
+	}
+#endif
+
 	USceneComponent* AttachComponent = SystemInstance->GetAttachComponent();
 	if (!AttachComponent)
 	{
