GeometryCache: Fixed GeometryCache ignoring "Receive Decals" setting with Deferred Decal materials.

diff --git a/Engine/Plugins/Runtime/GeometryCache/Source/GeometryCache/Private/GeometryCacheSceneProxy.cpp b/Engine/Plugins/Runtime/GeometryCache/Source/GeometryCache/Private/GeometryCacheSceneProxy.cpp
--- a/Engine/Plugins/Runtime/GeometryCache/Source/GeometryCache/Private/GeometryCacheSceneProxy.cpp
+++ b/Engine/Plugins/Runtime/GeometryCache/Source/GeometryCache/Private/GeometryCacheSceneProxy.cpp
@@ -399,7 +399,7 @@ void FGeometryCacheSceneProxy::CreateMeshBatch(
 	GetScene().GetPrimitiveUniformShaderParameters_RenderThread(GetPrimitiveSceneInfo(), bHasPrecomputedVolumetricLightmap, PreviousLocalToWorld, SingleCaptureIndex, bOutputVelocity);
 	bOutputVelocity |= AlwaysHasVelocity();
 
-	DynamicPrimitiveUniformBuffer.Set(LocalToWorldTransform, PreviousLocalToWorld, GetBounds(), GetLocalBounds(), true, false, bOutputVelocity);
+	DynamicPrimitiveUniformBuffer.Set(LocalToWorldTransform, PreviousLocalToWorld, GetBounds(), GetLocalBounds(), ReceivesDecals(), false, bOutputVelocity);
 	BatchElement.PrimitiveUniformBufferResource = &DynamicPrimitiveUniformBuffer.UniformBuffer;
 
 	const FGeometryCacheMeshData* MeshData = TrackProxy->bNextFrameMeshDataSelected ? TrackProxy->NextFrameMeshData : TrackProxy->MeshData;
