RHI: Fix assert in CIS where EBufferType::Typed buffers had no format specified.

#rnx

diff --git a/Engine/Plugins/Runtime/GeometryCache/Source/GeometryCache/Private/GeometryCacheSceneProxy.cpp b/Engine/Plugins/Runtime/GeometryCache/Source/GeometryCache/Private/GeometryCacheSceneProxy.cpp
--- a/Engine/Plugins/Runtime/GeometryCache/Source/GeometryCache/Private/GeometryCacheSceneProxy.cpp
+++ b/Engine/Plugins/Runtime/GeometryCache/Source/GeometryCache/Private/GeometryCacheSceneProxy.cpp
@@ -1488,7 +1488,7 @@ void FGeomCacheIndexBuffer::InitRHI()
 
 	if (IndexBufferRHI && NumAllocatedIndices)
 	{
-		BufferSRV = RHICreateShaderResourceView(IndexBufferRHI);
+		BufferSRV = RHICreateShaderResourceView(IndexBufferRHI, sizeof(uint32), PF_R32_UINT);
 	}
 }
 
@@ -1544,7 +1544,7 @@ void FGeomCacheIndexBuffer::Update(const TArray<uint32>& Indices)
 
 	if (bReallocate && IndexBufferRHI && NumAllocatedIndices)
 	{
-		BufferSRV = RHICreateShaderResourceView(IndexBufferRHI);
+		BufferSRV = RHICreateShaderResourceView(IndexBufferRHI, sizeof(uint32), PF_R32_UINT);
 	}
 }
 
@@ -1565,7 +1565,7 @@ void FGeomCacheIndexBuffer::UpdateSizeOnly(int32 NewNumIndices)
 
 	if (bReallocate && IndexBufferRHI && NumAllocatedIndices)
 	{
-		BufferSRV = RHICreateShaderResourceView(IndexBufferRHI);
+		BufferSRV = RHICreateShaderResourceView(IndexBufferRHI, sizeof(uint32), PF_R32_UINT);
 	}
 }
 
diff --git a/Engine/Source/Runtime/RHI/Public/RHICommandList.h b/Engine/Source/Runtime/RHI/Public/RHICommandList.h
--- a/Engine/Source/Runtime/RHI/Public/RHICommandList.h
+++ b/Engine/Source/Runtime/RHI/Public/RHICommandList.h
@@ -3990,9 +3990,11 @@ public:
 	//UE_DEPRECATED(5.3, "Use the CreateShaderResourceView function that takes an FRHIViewDesc.")
 	FORCEINLINE FShaderResourceViewRHIRef CreateShaderResourceView(FRHIBuffer* Buffer)
 	{
-		return CreateShaderResourceView(Buffer, FRHIViewDesc::CreateBufferSRV()
-			.SetTypeFromBuffer(Buffer)
-		);
+		FShaderResourceViewRHIRef SRVRef = CreateShaderResourceView(Buffer, FRHIViewDesc::CreateBufferSRV()
+				.SetTypeFromBuffer(Buffer));
+		checkf(SRVRef->GetDesc().Buffer.SRV.BufferType != FRHIViewDesc::EBufferType::Typed,
+			TEXT("Typed buffer should be created using CreateShaderResourceView where Format is specified."));
+		return SRVRef;
 	}
 	
 	//UE_DEPRECATED(5.3, "Use the CreateShaderResourceView function that takes an FRHIViewDesc.")
