- Add support for none particle bindings to geometry cache renderer
- Move accessor into properties to reduce CPU cost for each renderer

diff --git a/Engine/Plugins/Runtime/GeometryCache/Source/GeometryCache/Classes/NiagaraGeometryCacheRendererProperties.h b/Engine/Plugins/Runtime/GeometryCache/Source/GeometryCache/Classes/NiagaraGeometryCacheRendererProperties.h
--- a/Engine/Plugins/Runtime/GeometryCache/Source/GeometryCache/Classes/NiagaraGeometryCacheRendererProperties.h
+++ b/Engine/Plugins/Runtime/GeometryCache/Source/GeometryCache/Classes/NiagaraGeometryCacheRendererProperties.h
@@ -2,6 +2,7 @@
 
 #pragma once
 
+#include "NiagaraDataSetAccessor.h"
 #include "NiagaraRendererProperties.h"
 #include "NiagaraGeometryCacheRendererProperties.generated.h"
 
@@ -53,6 +54,7 @@ public:
 	virtual bool IsSimTargetSupported(ENiagaraSimTarget InSimTarget) const override { return InSimTarget == ENiagaraSimTarget::CPUSim; };
 	virtual void GetUsedMaterials(const FNiagaraEmitterInstance* InEmitter, TArray<UMaterialInterface*>& OutMaterials) const override;
 	virtual bool PopulateRequiredBindings(FNiagaraParameterStore& InParameterStore)  override;
+	virtual void CacheFromCompiledData(const FNiagaraDataSetCompiledData* CompiledData) override;
 
 #if WITH_EDITORONLY_DATA
 	virtual void GetRendererWidgets(const FNiagaraEmitterInstance* InEmitter, TArray<TSharedPtr<SWidget>>& OutWidgets, TSharedPtr<FAssetThumbnailPool> InThumbnailPool) const override;
@@ -111,6 +113,15 @@ public:
 	UPROPERTY(EditAnywhere, AdvancedDisplay, Category = "GeometryCache")
 	bool bAssignComponentsOnParticleID = true;
 
+	FNiagaraDataSetAccessor<FNiagaraPosition>	PositionAccessor;
+	FNiagaraDataSetAccessor<FVector3f>			RotationAccessor;
+	FNiagaraDataSetAccessor<FVector3f>			ScaleAccessor;
+	FNiagaraDataSetAccessor<float>				ElapsedTimeAccessor;
+	FNiagaraDataSetAccessor<FNiagaraBool>		EnabledAccessor;
+	FNiagaraDataSetAccessor<int32>				ArrayIndexAccessor;
+	FNiagaraDataSetAccessor<int32>				VisTagAccessor;
+	FNiagaraDataSetAccessor<int32>				UniqueIDAccessor;
+
 	virtual bool NeedsSystemPostTick() const override { return true; }
 	virtual bool NeedsSystemCompletion() const override { return true; }
 
diff --git a/Engine/Plugins/Runtime/GeometryCache/Source/GeometryCache/Private/NiagaraGeometryCacheRendererProperties.cpp b/Engine/Plugins/Runtime/GeometryCache/Source/GeometryCache/Private/NiagaraGeometryCacheRendererProperties.cpp
--- a/Engine/Plugins/Runtime/GeometryCache/Source/GeometryCache/Private/NiagaraGeometryCacheRendererProperties.cpp
+++ b/Engine/Plugins/Runtime/GeometryCache/Source/GeometryCache/Private/NiagaraGeometryCacheRendererProperties.cpp
@@ -249,6 +249,20 @@ bool UNiagaraGeometryCacheRendererProperties::PopulateRequiredBindings(FNiagaraP
 	return bAnyAdded;
 }
 
+void UNiagaraGeometryCacheRendererProperties::CacheFromCompiledData(const FNiagaraDataSetCompiledData* CompiledData)
+{
+	Super::CacheFromCompiledData(CompiledData);
+
+	InitParticleDataSetAccessor(PositionAccessor, CompiledData, PositionBinding);
+	InitParticleDataSetAccessor(RotationAccessor, CompiledData, RotationBinding);
+	InitParticleDataSetAccessor(ScaleAccessor, CompiledData, ScaleBinding);
+	InitParticleDataSetAccessor(ElapsedTimeAccessor, CompiledData, ElapsedTimeBinding);
+	InitParticleDataSetAccessor(EnabledAccessor, CompiledData, EnabledBinding);
+	InitParticleDataSetAccessor(ArrayIndexAccessor, CompiledData, ArrayIndexBinding);
+	InitParticleDataSetAccessor(VisTagAccessor, CompiledData, RendererVisibilityTagBinding);
+	UniqueIDAccessor.Init(CompiledData, FName("UniqueID"));
+}
+
 #if WITH_EDITORONLY_DATA
 
 const FSlateBrush* UNiagaraGeometryCacheRendererProperties::GetStackIcon() const
diff --git a/Engine/Plugins/Runtime/GeometryCache/Source/GeometryCache/Private/NiagaraRendererGeometryCache.cpp b/Engine/Plugins/Runtime/GeometryCache/Source/GeometryCache/Private/NiagaraRendererGeometryCache.cpp
--- a/Engine/Plugins/Runtime/GeometryCache/Source/GeometryCache/Private/NiagaraRendererGeometryCache.cpp
+++ b/Engine/Plugins/Runtime/GeometryCache/Source/GeometryCache/Private/NiagaraRendererGeometryCache.cpp
@@ -82,13 +82,13 @@ UGeometryCache* ResolveGeometryCache(const FNiagaraGeometryCacheReference& Entry
 	return FoundCache;
 }
 
-int32 GetGeometryCacheIndex(const FNiagaraDataSetReaderInt32<int32>& ArrayIndexAccessor, bool bCreateRandomIfUnassigned, const FNiagaraEmitterInstance* Emitter, const UNiagaraGeometryCacheRendererProperties* Properties, int32 ParticleIndex, int32 ParticleID)
+int32 GetGeometryCacheIndex(const FNiagaraDataSetReaderInt32<int32>& ArrayIndexReader, bool bCreateRandomIfUnassigned, const FNiagaraEmitterInstance* Emitter, const UNiagaraGeometryCacheRendererProperties* Properties, int32 ParticleIndex, int32 ParticleID)
 {
-	if (bCreateRandomIfUnassigned == false && ArrayIndexAccessor.IsValid() == false)
+	if (bCreateRandomIfUnassigned == false && ArrayIndexReader.IsValid() == false)
 	{
 		return INDEX_NONE;
 	}
-	int32 CacheIndex = ArrayIndexAccessor.GetSafe(ParticleIndex, INDEX_NONE);
+	int32 CacheIndex = ArrayIndexReader.GetSafe(ParticleIndex, INDEX_NONE);
 	if (CacheIndex == INDEX_NONE && Emitter->GetCachedEmitterData()->bDeterminism)
 	{
 		int32 Seed = Properties->bAssignComponentsOnParticleID ? ParticleID : ParticleIndex;
@@ -129,27 +129,26 @@ void FNiagaraRendererGeometryCache::PostSystemTick_GameThread(const UNiagaraRend
 		return;
 	}
 
+	const FNiagaraParameterStore& ParameterStore = Emitter->GetRendererBoundVariables();
+	const FNiagaraBool DefaultEnabled = ParameterStore.GetParameterValueOrDefault(Properties->EnabledBinding.GetParamMapBindableVariable(), FNiagaraBool(true));
+	const int32 DefaultVisTag = ParameterStore.GetParameterValueOrDefault(Properties->RendererVisibilityTagBinding.GetParamMapBindableVariable(), Properties->RendererVisibility);
+	//const int32 DefaultArrayIndex = ParameterStore.GetParameterValueOrDefault(Properties->ArrayIndexBinding.GetParamMapBindableVariable(), INDEX_NONE);	//-TODO: Revisit this when emitter mode is supported?
+	const int32 DefaultUniqueID = INDEX_NONE;
+
 	FNiagaraDataSet& Data = Emitter->GetData();
 	FNiagaraDataBuffer& ParticleData = Data.GetCurrentDataChecked();
-	FNiagaraDataSetReaderInt32<FNiagaraBool> EnabledAccessor = FNiagaraDataSetAccessor<FNiagaraBool>::CreateReader(Data, Properties->EnabledBinding.GetDataSetBindableVariable().GetName());
-	FNiagaraDataSetReaderInt32<int32> VisTagAccessor = FNiagaraDataSetAccessor<int32>::CreateReader(Data, Properties->RendererVisibilityTagBinding.GetDataSetBindableVariable().GetName());
-	FNiagaraDataSetReaderInt32<int32> ArrayIndexAccessor = FNiagaraDataSetAccessor<int32>::CreateReader(Data, Properties->ArrayIndexBinding.GetDataSetBindableVariable().GetName());
-	FNiagaraDataSetReaderInt32<int32> UniqueIDAccessor = FNiagaraDataSetAccessor<int32>::CreateReader(Data, FName("UniqueID"));
+	FNiagaraDataSetReaderInt32<FNiagaraBool> EnabledReader = Properties->EnabledAccessor.GetReader(Data);
+	FNiagaraDataSetReaderInt32<int32> VisTagReader = Properties->VisTagAccessor.GetReader(Data);
+	FNiagaraDataSetReaderInt32<int32> ArrayIndexReader = Properties->ArrayIndexAccessor.GetReader(Data);
+	FNiagaraDataSetReaderInt32<int32> UniqueIDReader = Properties->UniqueIDAccessor.GetReader(Data);
 	const float CurrentTime = AttachComponent->GetWorld()->GetRealTimeSeconds();
 	const bool bIsRendererEnabled = IsRendererEnabled(InProperties, Emitter);
 	
-	auto IsParticleEnabled = [&EnabledAccessor, &VisTagAccessor, Properties](int32 ParticleIndex)
-	{
-		if (EnabledAccessor.GetSafe(ParticleIndex, true))
+	const auto IsParticleEnabled =
+		[&EnabledReader, &DefaultEnabled, &VisTagReader, &DefaultVisTag, Properties](int32 ParticleIndex)
 		{
-			if (VisTagAccessor.IsValid())
-			{
-				return VisTagAccessor.GetSafe(ParticleIndex, 0) == Properties->RendererVisibility;
-			}
-			return true;
-		}
-		return false;
-	};
+			return EnabledReader.GetSafe(ParticleIndex, DefaultEnabled) && (VisTagReader.GetSafe(ParticleIndex, DefaultVisTag) == Properties->RendererVisibility);
+		};
 
 	TMap<int32, int32> ParticlesWithComponents;
 	TArray<int32> FreeList;
@@ -177,7 +176,7 @@ void FNiagaraRendererGeometryCache::PostSystemTick_GameThread(const UNiagaraRend
 		ParticlesWithComponents.Reserve(UsedSlots.Num());
 		for (uint32 ParticleIndex = 0; ParticleIndex < ParticleData.GetNumInstances(); ParticleIndex++)
 		{
-			int32 ParticleID = UniqueIDAccessor.GetSafe(ParticleIndex, -1);
+			int32 ParticleID = UniqueIDReader.GetSafe(ParticleIndex, DefaultUniqueID);
 			int32 PoolIndex;
 			if (UsedSlots.RemoveAndCopyValue(ParticleID, PoolIndex))
 			{
@@ -230,7 +229,7 @@ void FNiagaraRendererGeometryCache::PostSystemTick_GameThread(const UNiagaraRend
 		if (Properties->bAssignComponentsOnParticleID)
 		{
 			// Get the particle ID and see if we have any components already assigned to the particle
-			ParticleID = UniqueIDAccessor.GetSafe(ParticleIndex, -1);
+			ParticleID = UniqueIDReader.GetSafe(ParticleIndex, DefaultUniqueID);
 			ParticlesWithComponents.RemoveAndCopyValue(ParticleID, PoolIndex);
 		}
 
@@ -262,7 +261,7 @@ void FNiagaraRendererGeometryCache::PostSystemTick_GameThread(const UNiagaraRend
 		}
 
 		bool bCreateNewComponent = !GeometryComponent || GeometryComponent->HasAnyFlags(RF_BeginDestroyed | RF_FinishDestroyed);
-		int32 CacheIndex = GetGeometryCacheIndex(ArrayIndexAccessor, bCreateNewComponent, Emitter, Properties, ParticleIndex, ParticleID);
+		int32 CacheIndex = GetGeometryCacheIndex(ArrayIndexReader, bCreateNewComponent, Emitter, Properties, ParticleIndex, ParticleID);
 		if (Properties->GeometryCaches.IsValidIndex(CacheIndex))
 		{
 			const FNiagaraGeometryCacheReference& Entry = Properties->GeometryCaches[CacheIndex];
@@ -338,15 +337,20 @@ void FNiagaraRendererGeometryCache::PostSystemTick_GameThread(const UNiagaraRend
 			continue;
 		}
 
-		const auto PositionAccessor = FNiagaraDataSetAccessor<FNiagaraPosition>::CreateReader(Data, Properties->PositionBinding.GetDataSetBindableVariable().GetName());
-		const auto RotationAccessor = FNiagaraDataSetAccessor<FVector3f>::CreateReader(Data, Properties->RotationBinding.GetDataSetBindableVariable().GetName());
-		const auto ScaleAccessor = FNiagaraDataSetAccessor<FVector3f>::CreateReader(Data, Properties->ScaleBinding.GetDataSetBindableVariable().GetName());
-		const auto ElapsedTimeAccessor = FNiagaraDataSetAccessor<float>::CreateReader(Data, Properties->ElapsedTimeBinding.GetDataSetBindableVariable().GetName());
-		FNiagaraLWCConverter LwcConverter = SystemInstance->GetLWCConverter(Emitter->GetCachedEmitterData()->bLocalSpace);
+		const FNiagaraPosition DefaultPosition = ParameterStore.GetParameterValueOrDefault(Properties->PositionBinding.GetParamMapBindableVariable(), FNiagaraPosition(ForceInit));
+		const FVector3f DefaultRotation = ParameterStore.GetParameterValueOrDefault(Properties->RotationBinding.GetParamMapBindableVariable(), FVector3f::ZeroVector);
+		const FVector3f DefaultScale = ParameterStore.GetParameterValueOrDefault(Properties->ScaleBinding.GetParamMapBindableVariable(), FVector3f::OneVector);
+		const float DefaultElapsedTime = ParameterStore.GetParameterValueOrDefault(Properties->ElapsedTimeBinding.GetParamMapBindableVariable(), 0.0f);
+
+		const auto PositionAccessor = Properties->PositionAccessor.GetReader(Data);
+		const auto RotationAccessor = Properties->RotationAccessor.GetReader(Data);
+		const auto ScaleAccessor = Properties->ScaleAccessor.GetReader(Data);
+		const auto ElapsedTimeAccessor = Properties->ElapsedTimeAccessor.GetReader(Data);
+		const FNiagaraLWCConverter LwcConverter = SystemInstance->GetLWCConverter(Emitter->GetCachedEmitterData()->bLocalSpace);
 
-		FVector Position = LwcConverter.ConvertSimulationPositionToWorld(PositionAccessor.GetSafe(ParticleIndex, FNiagaraPosition(ForceInit)));
-		FVector Scale = (FVector)ScaleAccessor.GetSafe(ParticleIndex, FVector3f::OneVector);
-		FVector RotationVector = (FVector)RotationAccessor.GetSafe(ParticleIndex, FVector3f::ZeroVector);
+		FVector Position = LwcConverter.ConvertSimulationPositionToWorld(PositionAccessor.GetSafe(ParticleIndex, DefaultPosition));
+		FVector Scale = (FVector)ScaleAccessor.GetSafe(ParticleIndex, DefaultScale);
+		FVector RotationVector = (FVector)RotationAccessor.GetSafe(ParticleIndex, DefaultRotation);
 		FTransform Transform(FRotator(RotationVector.X, RotationVector.Y, RotationVector.Z), Position, Scale);
 		GeometryComponent->SetRelativeTransform(Transform);
 
@@ -357,7 +361,7 @@ void FNiagaraRendererGeometryCache::PostSystemTick_GameThread(const UNiagaraRend
 			GeometryComponent->Activate(false);
 		}
 
-		float ElapsedTime = ElapsedTimeAccessor.GetSafe(ParticleIndex, 0);
+		float ElapsedTime = ElapsedTimeAccessor.GetSafe(ParticleIndex, DefaultElapsedTime);
 
 		float Duration = GeometryComponent->GetDuration();
 		if (Properties->bIsLooping && ElapsedTime < 0)
