GeometryCache: Exposed MovieSceneGeometryCacheParams properties to Blueprint.

diff --git a/Engine/Plugins/Runtime/GeometryCache/Source/GeometryCacheTracks/Classes/MovieSceneGeometryCacheSection.h b/Engine/Plugins/Runtime/GeometryCache/Source/GeometryCacheTracks/Classes/MovieSceneGeometryCacheSection.h
--- a/Engine/Plugins/Runtime/GeometryCache/Source/GeometryCacheTracks/Classes/MovieSceneGeometryCacheSection.h
+++ b/Engine/Plugins/Runtime/GeometryCache/Source/GeometryCacheTracks/Classes/MovieSceneGeometryCacheSection.h
@@ -18,27 +18,27 @@ USTRUCT(BlueprintType)
 	/** Gets the animation sequence length, not modified by play rate */
 	 float GetSequenceLength() const;
 	/** The animation this section plays */
-	UPROPERTY(EditAnywhere, Category = "GeometryCache", DisplayName = "Geometry Cache")
+	UPROPERTY(BlueprintReadWrite, EditAnywhere, Category = "GeometryCache", DisplayName = "Geometry Cache")
 	TObjectPtr<UGeometryCache> GeometryCacheAsset;
 
 	/** The offset for the first loop of the animation clip */
-	UPROPERTY(EditAnywhere, Category = "GeometryCache")
+	UPROPERTY(BlueprintReadWrite, EditAnywhere, Category = "GeometryCache")
 	FFrameNumber FirstLoopStartFrameOffset;
 
 	/** The offset into the beginning of the animation clip */
-	UPROPERTY(EditAnywhere, Category = "GeometryCache")
+	UPROPERTY(BlueprintReadWrite, EditAnywhere, Category = "GeometryCache")
 	FFrameNumber StartFrameOffset;
 
 	/** The offset into the end of the animation clip */
-	UPROPERTY(EditAnywhere, Category = "GeometryCache")
+	UPROPERTY(BlueprintReadWrite, EditAnywhere, Category = "GeometryCache")
 	FFrameNumber EndFrameOffset;
 
 	/** The playback rate of the animation clip */
-	UPROPERTY(EditAnywhere, Category = "GeometryCache")
+	UPROPERTY(BlueprintReadWrite, EditAnywhere, Category = "GeometryCache")
 	float PlayRate;
 
 	/** Reverse the playback of the animation clip */
-	UPROPERTY(EditAnywhere, Category = "Animation")
+	UPROPERTY(BlueprintReadWrite, EditAnywhere, Category = "Animation")
 	uint32 bReverse : 1;
 
 	UPROPERTY()
@@ -62,7 +62,7 @@ class UMovieSceneGeometryCacheSection
 
 public:
 
-	UPROPERTY(EditAnywhere, Category = "Animation", meta = (ShowOnlyInnerProperties))
+	UPROPERTY(BlueprintReadWrite, EditAnywhere, Category = "Animation", meta = (ShowOnlyInnerProperties))
 	FMovieSceneGeometryCacheParams Params;
 
 	/** Get Frame Time as Animation Time*/
