- Geometry Cache Renderer Updates
- Emitter mode support added
- Material Bindings added

#rnx
#preflgiht 6461b199cf788a2558e9368f

diff --git a/Engine/Plugins/FX/Niagara/Source/Niagara/Public/NiagaraRendererProperties.h b/Engine/Plugins/FX/Niagara/Source/Niagara/Public/NiagaraRendererProperties.h
--- a/Engine/Plugins/FX/Niagara/Source/Niagara/Public/NiagaraRendererProperties.h
+++ b/Engine/Plugins/FX/Niagara/Source/Niagara/Public/NiagaraRendererProperties.h
@@ -252,9 +252,9 @@ struct FNiagaraRendererMaterialParameters
 	TArray<FNiagaraRendererMaterialStaticBoolParameter> StaticBoolParameters;
 
 #if WITH_EDITORONLY_DATA
-	void RenameVariable(const FNiagaraVariableBase& OldVariable, const FNiagaraVariableBase& NewVariable, const FVersionedNiagaraEmitter& InEmitter, ENiagaraRendererSourceDataMode SourceMode);
-	void RemoveVariable(const FNiagaraVariableBase& OldVariable, const FVersionedNiagaraEmitter& InEmitter, ENiagaraRendererSourceDataMode SourceMode);
-	void GetFeedback(TArrayView<UMaterialInterface*> Materials, TArray<FNiagaraRendererFeedback>& OutWarnings) const;
+	NIAGARA_API void RenameVariable(const FNiagaraVariableBase& OldVariable, const FNiagaraVariableBase& NewVariable, const FVersionedNiagaraEmitter& InEmitter, ENiagaraRendererSourceDataMode SourceMode);
+	NIAGARA_API void RemoveVariable(const FNiagaraVariableBase& OldVariable, const FVersionedNiagaraEmitter& InEmitter, ENiagaraRendererSourceDataMode SourceMode);
+	NIAGARA_API void GetFeedback(TArrayView<UMaterialInterface*> Materials, TArray<FNiagaraRendererFeedback>& OutWarnings) const;
 #endif
 
 	bool HasAnyBindings() const
diff --git a/Engine/Plugins/Runtime/GeometryCache/Source/GeometryCache/Classes/NiagaraGeometryCacheRendererProperties.h b/Engine/Plugins/Runtime/GeometryCache/Source/GeometryCache/Classes/NiagaraGeometryCacheRendererProperties.h
--- a/Engine/Plugins/Runtime/GeometryCache/Source/GeometryCache/Classes/NiagaraGeometryCacheRendererProperties.h
+++ b/Engine/Plugins/Runtime/GeometryCache/Source/GeometryCache/Classes/NiagaraGeometryCacheRendererProperties.h
@@ -6,6 +6,7 @@
 #include "NiagaraRendererProperties.h"
 #include "NiagaraGeometryCacheRendererProperties.generated.h"
 
+class UMaterialInstanceConstant;
 class UGeometryCache;
 struct FNiagaraVariable;
 struct FVersionedNiagaraEmitter;
@@ -31,6 +32,12 @@ public:
 	/** The materials to be used instead of the GeometryCache's materials. If the GeometryCache requires more materials than exist in this array or any entry in this array is set to None, we will use the GeometryCache's existing material instead.*/
 	UPROPERTY(EditAnywhere, Category = "GeometryCache")
 	TArray<TObjectPtr<UMaterialInterface>> OverrideMaterials;
+
+//-TODO: Follow what meshes do
+//#if WITH_EDITORONLY_DATA
+//	UPROPERTY(transient)
+//	TArray<TObjectPtr<UMaterialInstanceConstant>> MICMaterials;
+//#endif
 };
 
 UCLASS(editinlinenew, MinimalAPI, meta = (DisplayName = "Geometry Cache Renderer"))
@@ -44,6 +51,10 @@ public:
 	//UObject Interface
 	virtual void PostLoad() override;
 	virtual void PostInitProperties() override;
+	//virtual void Serialize(FArchive& Ar) override;
+#if WITH_EDITORONLY_DATA
+	virtual void PostEditChangeProperty(struct FPropertyChangedEvent& PropertyChangedEvent) override;
+#endif// WITH_EDITORONLY_DATA
 	//UObject Interface END
 
 	static void InitCDOPropertiesAfterModuleStartup();
@@ -53,7 +64,10 @@ public:
 	virtual class FNiagaraBoundsCalculator* CreateBoundsCalculator() override { return nullptr; }
 	virtual bool IsSimTargetSupported(ENiagaraSimTarget InSimTarget) const override { return InSimTarget == ENiagaraSimTarget::CPUSim; };
 	virtual void GetUsedMaterials(const FNiagaraEmitterInstance* InEmitter, TArray<UMaterialInterface*>& OutMaterials) const override;
+	virtual void UpdateSourceModeDerivates(ENiagaraRendererSourceDataMode InSourceMode, bool bFromPropertyEdit = false) override;
+	virtual ENiagaraRendererSourceDataMode GetCurrentSourceMode() const override { return SourceMode; }
 	virtual bool PopulateRequiredBindings(FNiagaraParameterStore& InParameterStore)  override;
+	virtual bool NeedsMIDsForMaterials() const { return MaterialParameters.HasAnyBindings(); }
 	virtual void CacheFromCompiledData(const FNiagaraDataSetCompiledData* CompiledData) override;
 
 #if WITH_EDITORONLY_DATA
@@ -61,13 +75,24 @@ public:
 	virtual const FSlateBrush* GetStackIcon() const override;
 	virtual void GetRendererTooltipWidgets(const FNiagaraEmitterInstance* InEmitter, TArray<TSharedPtr<SWidget>>& OutWidgets, TSharedPtr<FAssetThumbnailPool> InThumbnailPool) const override;
 	virtual const TArray<FNiagaraVariable>& GetOptionalAttributes() override;
-	virtual void GetRendererFeedback(const FVersionedNiagaraEmitter& InEmitter, TArray<FText>& OutErrors, TArray<FText>& OutWarnings, TArray<FText>& OutInfo) const override;
+	virtual void GetRendererFeedback(const FVersionedNiagaraEmitter& InEmitter, TArray<FNiagaraRendererFeedback>& OutErrors, TArray<FNiagaraRendererFeedback>& OutWarnings, TArray<FNiagaraRendererFeedback>& OutInfo) const override;
+	virtual TArray<FNiagaraVariable> GetBoundAttributes() const override;
+	virtual void RenameVariable(const FNiagaraVariableBase& OldVariable, const FNiagaraVariableBase& NewVariable, const FVersionedNiagaraEmitter& InEmitter) override;
+	virtual void RemoveVariable(const FNiagaraVariableBase& OldVariable, const FVersionedNiagaraEmitter& InEmitter) override;
 #endif // WITH_EDITORONLY_DATA
 
+	void UpdateMICs();
+
+	static UGeometryCache* ResolveGeometryCache(const FNiagaraGeometryCacheReference& Entry, const FNiagaraEmitterInstance* Emitter);
+
 	/** Reference to the geometry cache assets to use. If ArrayIndexBinding is not set, a random element is used for each particle. */
 	UPROPERTY(EditAnywhere, Category = "GeometryCache")
 	TArray<FNiagaraGeometryCacheReference> GeometryCaches;
 
+	/** Whether or not to draw a single element for the Emitter or to draw the particles.*/
+	UPROPERTY(EditAnywhere, Category = "GeometryCache")
+	ENiagaraRendererSourceDataMode SourceMode = ENiagaraRendererSourceDataMode::Particles;
+
 	/** If true, then the geometry cache keeps playing in a loop */
 	UPROPERTY(EditAnywhere, Category = "GeometryCache")
 	bool bIsLooping = true;
@@ -113,6 +138,10 @@ public:
 	UPROPERTY(EditAnywhere, AdvancedDisplay, Category = "GeometryCache")
 	bool bAssignComponentsOnParticleID = true;
 
+	/** If this array has entries, we will create a MaterialInstanceDynamic per Emitter instance from Material and set the Material parameters using the Niagara simulation variables listed.*/
+	UPROPERTY(EditAnywhere, Category = "Bindings")
+	FNiagaraRendererMaterialParameters MaterialParameters;
+
 	FNiagaraDataSetAccessor<FNiagaraPosition>	PositionAccessor;
 	FNiagaraDataSetAccessor<FVector3f>			RotationAccessor;
 	FNiagaraDataSetAccessor<FVector3f>			ScaleAccessor;
diff --git a/Engine/Plugins/Runtime/GeometryCache/Source/GeometryCache/Classes/NiagaraRendererGeometryCache.h b/Engine/Plugins/Runtime/GeometryCache/Source/GeometryCache/Classes/NiagaraRendererGeometryCache.h
--- a/Engine/Plugins/Runtime/GeometryCache/Source/GeometryCache/Classes/NiagaraRendererGeometryCache.h
+++ b/Engine/Plugins/Runtime/GeometryCache/Source/GeometryCache/Classes/NiagaraRendererGeometryCache.h
@@ -23,6 +23,7 @@ public:
 	virtual ~FNiagaraRendererGeometryCache();
 
 	//FNiagaraRenderer interface
+	virtual void Initialize(const UNiagaraRendererProperties* InProps, const FNiagaraEmitterInstance* Emitter, const FNiagaraSystemInstanceController& InController) override;
 	virtual void DestroyRenderState_Concurrent() override;
 	virtual void PostSystemTick_GameThread(const UNiagaraRendererProperties* InProperties, const FNiagaraEmitterInstance* Emitter) override;
 	virtual void OnSystemComplete_GameThread(const UNiagaraRendererProperties* InProperties, const FNiagaraEmitterInstance* Emitter) override;
@@ -37,6 +38,16 @@ private:
 		int32 LastAssignedToParticleID = -1;
 	};
 
+	UGeometryCacheComponent* CreateOrGetPooledComponent(
+		const UNiagaraGeometryCacheRendererProperties* Properties,
+		USceneComponent* AttachComponent,
+		const FNiagaraEmitterInstance* Emitter,
+		TOptional<int32> DefaultCacheIndex,
+		uint32 ParticleIndex,
+		int32 ParticleID,
+		int32& InOutPoolIndex
+	);
+
 	void ResetComponentPool(bool bResetOwner);
 
 	// if the niagara component is not attached to an actor, we need to spawn and keep track of a temporary actor
@@ -44,6 +55,9 @@ private:
 
 	// all of the spawned components
 	TArray<FComponentPoolEntry> ComponentPool;
+
+	// Used to remap materials into the compact list so we can use MIDs
+	TArray<TArray<int32>> MaterialRemapTable;
 };
 
 #if UE_ENABLE_INCLUDE_ORDER_DEPRECATED_IN_5_2
diff --git a/Engine/Plugins/Runtime/GeometryCache/Source/GeometryCache/Private/NiagaraGeometryCacheRendererProperties.cpp b/Engine/Plugins/Runtime/GeometryCache/Source/GeometryCache/Private/NiagaraGeometryCacheRendererProperties.cpp
--- a/Engine/Plugins/Runtime/GeometryCache/Source/GeometryCache/Private/NiagaraGeometryCacheRendererProperties.cpp
+++ b/Engine/Plugins/Runtime/GeometryCache/Source/GeometryCache/Private/NiagaraGeometryCacheRendererProperties.cpp
@@ -3,6 +3,8 @@
 #include "NiagaraGeometryCacheRendererProperties.h"
 #include "GeometryCache.h"
 #include "NiagaraConstants.h"
+#include "NiagaraEmitter.h"
+#include "NiagaraEmitterInstance.h"
 #include "NiagaraModule.h"
 #include "NiagaraRendererGeometryCache.h"
 #include "Modules/ModuleManager.h"
@@ -98,6 +100,10 @@ void UNiagaraGeometryCacheRendererProperties::InitDefaultAttributes()
 	}
 }
 
+//void UNiagaraGeometryCacheRendererProperties::Serialize(FArchive& Ar)
+//{
+//}
+
 void UNiagaraGeometryCacheRendererProperties::PostInitProperties()
 {
 	Super::PostInitProperties();
@@ -145,6 +151,29 @@ void UNiagaraGeometryCacheRendererProperties::PostInitProperties()
 	}
 }
 
+#if WITH_EDITORONLY_DATA
+void UNiagaraGeometryCacheRendererProperties::PostEditChangeProperty(struct FPropertyChangedEvent& PropertyChangedEvent)
+{
+	Super::PostEditChangeProperty(PropertyChangedEvent);
+
+	const FName PropertyName = PropertyChangedEvent.GetPropertyName();
+	const FName MemberPropertyName = PropertyChangedEvent.GetMemberPropertyName();
+
+	if (PropertyName == GET_MEMBER_NAME_CHECKED(UNiagaraGeometryCacheRendererProperties, SourceMode))
+	{
+		UpdateSourceModeDerivates(SourceMode, true);
+	}
+
+	// Update our MICs if we change material / material bindings
+	//-OPT: Could narrow down further to only static materials
+	//if ((PropertyName == GET_MEMBER_NAME_CHECKED(UNiagaraGeometryCacheRendererProperties, Material)) ||
+	//	(MemberPropertyName == GET_MEMBER_NAME_CHECKED(UNiagaraGeometryCacheRendererProperties, MaterialParameters)))
+	//{
+	//	UpdateMICs();
+	//}
+}
+#endif// WITH_EDITORONLY_DATA
+
 void UNiagaraGeometryCacheRendererProperties::InitCDOPropertiesAfterModuleStartup()
 {
 	InitDefaultAttributes();
@@ -201,28 +230,44 @@ FNiagaraRenderer* UNiagaraGeometryCacheRendererProperties::CreateEmitterRenderer
 	return NewRenderer;
 }
 
-void UNiagaraGeometryCacheRendererProperties::GetUsedMaterials(const FNiagaraEmitterInstance* InEmitter, TArray<UMaterialInterface*>& MaterialInterfaces) const
+void UNiagaraGeometryCacheRendererProperties::GetUsedMaterials(const FNiagaraEmitterInstance* EmitterInstance, TArray<UMaterialInterface*>& MaterialInterfaces) const
 {
 	for (const FNiagaraGeometryCacheReference& Entry : GeometryCaches)
 	{
-		if (Entry.GeometryCache)
+		UGeometryCache* GeometryCache = ResolveGeometryCache(Entry, EmitterInstance);;
+		if (GeometryCache == nullptr)
+		{
+			continue;
+		}
+
+		const int32 MaxIndex = FMath::Max(GeometryCache->Materials.Num(), Entry.OverrideMaterials.Num());
+		for (int i = 0; i < MaxIndex; i++)
 		{
-			int32 MaxIndex = FMath::Max(Entry.GeometryCache->Materials.Num(), Entry.OverrideMaterials.Num());
-			for (int i = 0; i < MaxIndex; i++)
+			if (Entry.OverrideMaterials.IsValidIndex(i) && Entry.OverrideMaterials[i])
+			{
+				MaterialInterfaces.AddUnique(Entry.OverrideMaterials[i]);
+			}
+			else if (GeometryCache->Materials.IsValidIndex(i))
 			{
-				if (Entry.OverrideMaterials.IsValidIndex(i) && Entry.OverrideMaterials[i])
-				{
-					MaterialInterfaces.Add(Entry.OverrideMaterials[i]);
-				}
-				else if (Entry.GeometryCache->Materials.IsValidIndex(i))
-				{
-					MaterialInterfaces.Add(Entry.GeometryCache->Materials[i]);
-				}
+				MaterialInterfaces.AddUnique(GeometryCache->Materials[i]);
 			}
 		}
 	}
 }
 
+void UNiagaraGeometryCacheRendererProperties::UpdateSourceModeDerivates(ENiagaraRendererSourceDataMode InSourceMode, bool bFromPropertyEdit)
+{
+	Super::UpdateSourceModeDerivates(InSourceMode, bFromPropertyEdit);
+
+	if (UNiagaraEmitter* SrcEmitter = GetTypedOuter<UNiagaraEmitter>())
+	{
+		for (FNiagaraMaterialAttributeBinding& MaterialParamBinding : MaterialParameters.AttributeBindings)
+		{
+			MaterialParamBinding.CacheValues(SrcEmitter);
+		}
+	}
+}
+
 bool UNiagaraGeometryCacheRendererProperties::PopulateRequiredBindings(FNiagaraParameterStore& InParameterStore)
 {
 	bool bAnyAdded = Super::PopulateRequiredBindings(InParameterStore);
@@ -246,6 +291,12 @@ bool UNiagaraGeometryCacheRendererProperties::PopulateRequiredBindings(FNiagaraP
 		}
 	}
 
+	for (const FNiagaraMaterialAttributeBinding& MaterialParamBinding : MaterialParameters.AttributeBindings)
+	{
+		InParameterStore.AddParameter(MaterialParamBinding.GetParamMapBindableVariable(), false);
+		bAnyAdded = true;
+	}
+
 	return bAnyAdded;
 }
 
@@ -253,6 +304,9 @@ void UNiagaraGeometryCacheRendererProperties::CacheFromCompiledData(const FNiaga
 {
 	Super::CacheFromCompiledData(CompiledData);
 
+	UpdateSourceModeDerivates(SourceMode);
+	UpdateMICs();
+
 	InitParticleDataSetAccessor(PositionAccessor, CompiledData, PositionBinding);
 	InitParticleDataSetAccessor(RotationAccessor, CompiledData, RotationBinding);
 	InitParticleDataSetAccessor(ScaleAccessor, CompiledData, ScaleBinding);
@@ -264,7 +318,6 @@ void UNiagaraGeometryCacheRendererProperties::CacheFromCompiledData(const FNiaga
 }
 
 #if WITH_EDITORONLY_DATA
-
 const FSlateBrush* UNiagaraGeometryCacheRendererProperties::GetStackIcon() const
 {
 	return FSlateIconFinder::FindIconBrushForClass(GetClass());
@@ -350,11 +403,74 @@ const TArray<FNiagaraVariable>& UNiagaraGeometryCacheRendererProperties::GetOpti
 	return Attrs;
 }
 
-void UNiagaraGeometryCacheRendererProperties::GetRendererFeedback(const FVersionedNiagaraEmitter& InEmitter, TArray<FText>& OutErrors, TArray<FText>& OutWarnings, TArray<FText>& OutInfo) const
+void UNiagaraGeometryCacheRendererProperties::GetRendererFeedback(const FVersionedNiagaraEmitter& InEmitter, TArray<FNiagaraRendererFeedback>& OutErrors, TArray<FNiagaraRendererFeedback>& OutWarnings, TArray<FNiagaraRendererFeedback>& OutInfo) const
 {
 	OutInfo.Add(FText::FromString(TEXT("The geometry cache renderer is still an experimental feature.")));
+
+	if (MaterialParameters.HasAnyBindings())
+	{
+		TArray<UMaterialInterface*> Materials;
+		GetUsedMaterials(nullptr, Materials);
+		MaterialParameters.GetFeedback(Materials, OutWarnings);
+	}
+}
+
+TArray<FNiagaraVariable> UNiagaraGeometryCacheRendererProperties::GetBoundAttributes() const
+{
+	TArray<FNiagaraVariable> BoundAttributes = Super::GetBoundAttributes();
+	BoundAttributes.Reserve(BoundAttributes.Num() + MaterialParameters.AttributeBindings.Num());
+
+	for (const FNiagaraMaterialAttributeBinding& MaterialParamBinding : MaterialParameters.AttributeBindings)
+	{
+		BoundAttributes.AddUnique(MaterialParamBinding.GetParamMapBindableVariable());
+	}
+	return BoundAttributes;
+}
+
+void UNiagaraGeometryCacheRendererProperties::RenameVariable(const FNiagaraVariableBase& OldVariable, const FNiagaraVariableBase& NewVariable, const FVersionedNiagaraEmitter& InEmitter)
+{
+	Super::RenameVariable(OldVariable, NewVariable, InEmitter);
+#if WITH_EDITORONLY_DATA
+	MaterialParameters.RenameVariable(OldVariable, NewVariable, InEmitter, GetCurrentSourceMode());
+#endif
+}
+
+void UNiagaraGeometryCacheRendererProperties::RemoveVariable(const FNiagaraVariableBase& OldVariable, const FVersionedNiagaraEmitter& InEmitter)
+{
+	Super::RemoveVariable(OldVariable, InEmitter);
+#if WITH_EDITORONLY_DATA
+	MaterialParameters.RemoveVariable(OldVariable, InEmitter, GetCurrentSourceMode());
+#endif
 }
+#endif
 
+void UNiagaraGeometryCacheRendererProperties::UpdateMICs()
+{
+#if WITH_EDITORONLY_DATA
+	//UpdateMaterialParametersMIC(MaterialParameters, Material, MICMaterial);
 #endif
+}
+
+UGeometryCache* UNiagaraGeometryCacheRendererProperties::ResolveGeometryCache(const FNiagaraGeometryCacheReference& Entry, const FNiagaraEmitterInstance* Emitter)
+{
+	UGeometryCache* FoundCache = nullptr;
+
+	FNiagaraVariable Variable = Entry.GeometryCacheUserParamBinding.Parameter;
+	if (Variable.IsValid() && Emitter)
+	{
+		UGeometryCache* GeometryCache = Cast<UGeometryCache>(Emitter->GetRendererBoundVariables().GetUObject(Variable));
+		if (GeometryCache)
+		{
+			FoundCache = GeometryCache;
+		}
+	}
+
+	if (!FoundCache)
+	{
+		FoundCache = Entry.GeometryCache;
+	}
+
+	return FoundCache;
+}
 
 #undef LOCTEXT_NAMESPACE
diff --git a/Engine/Plugins/Runtime/GeometryCache/Source/GeometryCache/Private/NiagaraRendererGeometryCache.cpp b/Engine/Plugins/Runtime/GeometryCache/Source/GeometryCache/Private/NiagaraRendererGeometryCache.cpp
--- a/Engine/Plugins/Runtime/GeometryCache/Source/GeometryCache/Private/NiagaraRendererGeometryCache.cpp
+++ b/Engine/Plugins/Runtime/GeometryCache/Source/GeometryCache/Private/NiagaraRendererGeometryCache.cpp
@@ -10,6 +10,8 @@
 #include "NiagaraEmitterInstance.h"
 #include "NiagaraSystemInstance.h"
 #include "Async/Async.h"
+#include "Materials/Material.h"
+#include "Materials/MaterialInstanceDynamic.h"
 
 static float GNiagaraGeometryComponentRenderPoolInactiveTimeLimit = 5;
 static FAutoConsoleVariableRef CVarNiagaraComponentRenderPoolInactiveTimeLimit(
@@ -35,6 +37,45 @@ FNiagaraRendererGeometryCache::~FNiagaraRendererGeometryCache()
 	check(ComponentPool.Num() == 0);
 }
 
+void FNiagaraRendererGeometryCache::Initialize(const UNiagaraRendererProperties* InProps, const FNiagaraEmitterInstance* Emitter, const FNiagaraSystemInstanceController& InController)
+{
+	FNiagaraRenderer::Initialize(InProps, Emitter, InController);
+
+	const UNiagaraGeometryCacheRendererProperties* Properties = CastChecked<UNiagaraGeometryCacheRendererProperties>(InProps);
+
+	// Iniitialize the cache -> material mappings
+	MaterialRemapTable.SetNum(Properties->GeometryCaches.Num());
+
+	for (int32 iCache=0; iCache < Properties->GeometryCaches.Num(); ++iCache)
+	{
+		const FNiagaraGeometryCacheReference& GeoCacheRef = Properties->GeometryCaches[iCache];
+		UGeometryCache* GeoCache = UNiagaraGeometryCacheRendererProperties::ResolveGeometryCache(GeoCacheRef, Emitter);
+		if (GeoCache == nullptr)
+		{
+			continue;
+		}
+		MaterialRemapTable[iCache].SetNum(GeoCache->Materials.Num());
+
+		for ( int32 iMaterial=0; iMaterial < GeoCache->Materials.Num(); ++iMaterial)
+		{
+			UMaterialInterface* Material = GeoCacheRef.OverrideMaterials.IsValidIndex(iMaterial) ? GeoCacheRef.OverrideMaterials[iMaterial] : GeoCache->Materials[iMaterial];
+
+			MaterialRemapTable[iCache][iMaterial] =
+				BaseMaterials_GT.IndexOfByPredicate(
+					[Material](UMaterialInterface* UsedMaterial)
+					{
+						if (UsedMaterial == Material)
+						{
+							return true;
+						}
+						UMaterialInstanceDynamic* MID = Cast<UMaterialInstanceDynamic>(UsedMaterial);
+						return MID ? MID->Parent == Material : false;
+					}
+				);
+		}
+	}
+}
+
 void FNiagaraRendererGeometryCache::DestroyRenderState_Concurrent()
 {
 	// Rendering resources are being destroyed, but the component pool and their owner actor must be destroyed on the game thread
@@ -60,35 +101,13 @@ void FNiagaraRendererGeometryCache::DestroyRenderState_Concurrent()
 	SpawnedOwner.Reset();
 }
 
-UGeometryCache* ResolveGeometryCache(const FNiagaraGeometryCacheReference& Entry, const FNiagaraEmitterInstance* Emitter)
+int32 GetGeometryCacheIndex(TOptional<int32> DefaultCacheIndex, bool bCreateRandomIfUnassigned, const FNiagaraEmitterInstance* Emitter, const UNiagaraGeometryCacheRendererProperties* Properties, int32 ParticleIndex, int32 ParticleID)
 {
-	UGeometryCache* FoundCache = nullptr;
-
-	FNiagaraVariable Variable = Entry.GeometryCacheUserParamBinding.Parameter;
-	if (Variable.IsValid() && Emitter)
-	{
-		UGeometryCache* GeometryCache = Cast<UGeometryCache>(Emitter->GetRendererBoundVariables().GetUObject(Variable));
-		if (GeometryCache)
-		{
-			FoundCache = GeometryCache;
-		}
-	}
-
-	if (!FoundCache)
-	{
-		FoundCache = Entry.GeometryCache;
-	}
-
-	return FoundCache;
-}
-
-int32 GetGeometryCacheIndex(const FNiagaraDataSetReaderInt32<int32>& ArrayIndexReader, bool bCreateRandomIfUnassigned, const FNiagaraEmitterInstance* Emitter, const UNiagaraGeometryCacheRendererProperties* Properties, int32 ParticleIndex, int32 ParticleID)
-{
-	if (bCreateRandomIfUnassigned == false && ArrayIndexReader.IsValid() == false)
+	if (bCreateRandomIfUnassigned == false && !DefaultCacheIndex.IsSet())
 	{
 		return INDEX_NONE;
 	}
-	int32 CacheIndex = ArrayIndexReader.GetSafe(ParticleIndex, INDEX_NONE);
+	int32 CacheIndex = DefaultCacheIndex.Get(INDEX_NONE);
 	if (CacheIndex == INDEX_NONE && Emitter->GetCachedEmitterData()->bDeterminism)
 	{
 		int32 Seed = Properties->bAssignComponentsOnParticleID ? ParticleID : ParticleIndex;
@@ -102,6 +121,99 @@ int32 GetGeometryCacheIndex(const FNiagaraDataSetReaderInt32<int32>& ArrayIndexR
 	return CacheIndex % Properties->GeometryCaches.Num(); // wrap around if index is too large
 }
 
+UGeometryCacheComponent* FNiagaraRendererGeometryCache::CreateOrGetPooledComponent(
+	const UNiagaraGeometryCacheRendererProperties* Properties,
+	USceneComponent* AttachComponent,
+	const FNiagaraEmitterInstance* Emitter,
+	TOptional<int32> DefaultCacheIndex,
+	uint32 ParticleIndex,
+	int32 ParticleID,
+	int32& InOutPoolIndex
+)
+{
+	UGeometryCacheComponent* GeometryComponent = ComponentPool.IsValidIndex(InOutPoolIndex) ? ComponentPool[InOutPoolIndex].Component.Get() : nullptr;
+	const bool bCreateNewComponent = !GeometryComponent || GeometryComponent->HasAnyFlags(RF_BeginDestroyed | RF_FinishDestroyed);
+
+	const int32 CacheIndex = GetGeometryCacheIndex(DefaultCacheIndex, bCreateNewComponent, Emitter, Properties, ParticleIndex, ParticleID);
+	if (Properties->GeometryCaches.IsValidIndex(CacheIndex))
+	{
+		const FNiagaraGeometryCacheReference& Entry = Properties->GeometryCaches[CacheIndex];
+		UGeometryCache* Cache = UNiagaraGeometryCacheRendererProperties::ResolveGeometryCache(Entry, Emitter);
+		if (!Cache)
+		{
+			return nullptr;
+		}
+
+		if (bCreateNewComponent)
+		{
+			// Determine the owner actor or spawn one
+			AActor* OwnerActor = SpawnedOwner.Get();
+			if (OwnerActor == nullptr)
+			{
+				OwnerActor = AttachComponent->GetOwner();
+				if (OwnerActor == nullptr)
+				{
+					// NOTE: This can happen with spawned systems
+					OwnerActor = AttachComponent->GetWorld()->SpawnActor<AActor>();
+					OwnerActor->SetFlags(RF_Transient);
+					SpawnedOwner = OwnerActor;
+				}
+			}
+
+			// if we don't have a pooled component we create a new one from the template
+			GeometryComponent = NewObject<UGeometryCacheComponent>(OwnerActor);
+			IGeometryCacheStreamingManager::Get().AddStreamingComponent(GeometryComponent);
+			GeometryComponent->SetFlags(RF_Transient);
+			GeometryComponent->SetupAttachment(AttachComponent);
+			GeometryComponent->RegisterComponent();
+			GeometryComponent->AddTickPrerequisiteComponent(AttachComponent);
+			GeometryComponent->SetLooping(Properties->bIsLooping);
+			GeometryComponent->SetManualTick(true); // we want to tick the component with the delta time of the niagara sim
+
+			if (Emitter->GetCachedEmitterData()->bLocalSpace)
+			{
+				GeometryComponent->SetAbsolute(false, false, false);
+			}
+			else
+			{
+				GeometryComponent->SetAbsolute(true, true, true);
+			}
+
+			if (InOutPoolIndex >= 0)
+			{
+				// This should only happen if the component was destroyed externally
+				ComponentPool[InOutPoolIndex].Component = GeometryComponent;
+				ComponentPool[InOutPoolIndex].LastElapsedTime = 0;
+			}
+			else
+			{
+				// Add a new pool entry
+				InOutPoolIndex = ComponentPool.Num();
+				ComponentPool.AddDefaulted_GetRef().Component = GeometryComponent;
+				ComponentPool[InOutPoolIndex].LastElapsedTime = 0;
+			}
+		}
+
+		GeometryComponent->SetGeometryCache(Cache);
+
+		// set the cache and override materials, since they might have changed from the last time we accessed the array index
+		// if they did not change from last tick these calls should be no-ops.
+		for (int32 i=0; i < MaterialRemapTable[CacheIndex].Num(); ++i)
+		{
+			const int32 iMaterial = MaterialRemapTable[CacheIndex][i];
+			UMaterialInterface* Material = BaseMaterials_GT.IsValidIndex(iMaterial) ? BaseMaterials_GT[iMaterial] : nullptr;
+			GeometryComponent->SetMaterial(i, Material);
+		}
+	}
+	else if (bCreateNewComponent)
+	{
+		return nullptr;
+	}
+
+	return GeometryComponent;
+}
+
+
 /** Update render data buffer from attributes */
 void FNiagaraRendererGeometryCache::PostSystemTick_GameThread(const UNiagaraRendererProperties* InProperties, const FNiagaraEmitterInstance* Emitter)
 {
@@ -129,302 +241,294 @@ void FNiagaraRendererGeometryCache::PostSystemTick_GameThread(const UNiagaraRend
 		return;
 	}
 
+	if (Properties->MaterialParameters.HasAnyBindings())
+	{
+		ProcessMaterialParameterBindings(Properties->MaterialParameters, Emitter, MakeArrayView(BaseMaterials_GT));
+	}
+
 	const FNiagaraParameterStore& ParameterStore = Emitter->GetRendererBoundVariables();
 	const FNiagaraBool DefaultEnabled = ParameterStore.GetParameterValueOrDefault(Properties->EnabledBinding.GetParamMapBindableVariable(), FNiagaraBool(true));
 	const int32 DefaultVisTag = ParameterStore.GetParameterValueOrDefault(Properties->RendererVisibilityTagBinding.GetParamMapBindableVariable(), Properties->RendererVisibility);
 	//const int32 DefaultArrayIndex = ParameterStore.GetParameterValueOrDefault(Properties->ArrayIndexBinding.GetParamMapBindableVariable(), INDEX_NONE);	//-TODO: Revisit this when emitter mode is supported?
 	const int32 DefaultUniqueID = INDEX_NONE;
-
-	FNiagaraDataSet& Data = Emitter->GetData();
-	FNiagaraDataBuffer& ParticleData = Data.GetCurrentDataChecked();
-	FNiagaraDataSetReaderInt32<FNiagaraBool> EnabledReader = Properties->EnabledAccessor.GetReader(Data);
-	FNiagaraDataSetReaderInt32<int32> VisTagReader = Properties->VisTagAccessor.GetReader(Data);
-	FNiagaraDataSetReaderInt32<int32> ArrayIndexReader = Properties->ArrayIndexAccessor.GetReader(Data);
-	FNiagaraDataSetReaderInt32<int32> UniqueIDReader = Properties->UniqueIDAccessor.GetReader(Data);
+	const FNiagaraPosition DefaultPosition = ParameterStore.GetParameterValueOrDefault(Properties->PositionBinding.GetParamMapBindableVariable(), FNiagaraPosition(ForceInit));
+	const FVector3f DefaultRotation = ParameterStore.GetParameterValueOrDefault(Properties->RotationBinding.GetParamMapBindableVariable(), FVector3f::ZeroVector);
+	const FVector3f DefaultScale = ParameterStore.GetParameterValueOrDefault(Properties->ScaleBinding.GetParamMapBindableVariable(), FVector3f::OneVector);
+	const float DefaultElapsedTime = ParameterStore.GetParameterValueOrDefault(Properties->ElapsedTimeBinding.GetParamMapBindableVariable(), 0.0f);
 	const float CurrentTime = AttachComponent->GetWorld()->GetRealTimeSeconds();
+
+	const FNiagaraLWCConverter LwcConverter = SystemInstance->GetLWCConverter(Emitter->GetCachedEmitterData()->bLocalSpace);
 	const bool bIsRendererEnabled = IsRendererEnabled(InProperties, Emitter);
-	
-	const auto IsParticleEnabled =
-		[&EnabledReader, &DefaultEnabled, &VisTagReader, &DefaultVisTag, Properties](int32 ParticleIndex)
-		{
-			return EnabledReader.GetSafe(ParticleIndex, DefaultEnabled) && (VisTagReader.GetSafe(ParticleIndex, DefaultVisTag) == Properties->RendererVisibility);
-		};
 
-	TMap<int32, int32> ParticlesWithComponents;
-	TArray<int32> FreeList;
-	if (Properties->bAssignComponentsOnParticleID && ComponentPool.Num() > 0)
+	if (Properties->SourceMode == ENiagaraRendererSourceDataMode::Particles)
 	{
-		FreeList.Reserve(ComponentPool.Num());
-
-		// Determine the slots that were assigned to particles last frame
-		TMap<int32, int32> UsedSlots;
-		UsedSlots.Reserve(ComponentPool.Num());
-		for (int32 EntryIndex = 0; EntryIndex < ComponentPool.Num(); ++EntryIndex)
-		{
-			FComponentPoolEntry& Entry = ComponentPool[EntryIndex];
-			if (Entry.LastAssignedToParticleID >= 0)
-			{
-				UsedSlots.Emplace(Entry.LastAssignedToParticleID, EntryIndex);
-			}
-			else
+		FNiagaraDataSet& Data = Emitter->GetData();
+		FNiagaraDataBuffer& ParticleData = Data.GetCurrentDataChecked();
+		FNiagaraDataSetReaderInt32<FNiagaraBool> EnabledReader = Properties->EnabledAccessor.GetReader(Data);
+		FNiagaraDataSetReaderInt32<int32> VisTagReader = Properties->VisTagAccessor.GetReader(Data);
+		FNiagaraDataSetReaderInt32<int32> ArrayIndexReader = Properties->ArrayIndexAccessor.GetReader(Data);
+		FNiagaraDataSetReaderInt32<int32> UniqueIDReader = Properties->UniqueIDAccessor.GetReader(Data);
+
+		const auto IsParticleEnabled =
+			[&EnabledReader, &DefaultEnabled, &VisTagReader, &DefaultVisTag, Properties](int32 ParticleIndex)
 			{
-				FreeList.Add(EntryIndex);
-			}
-		}
+				return EnabledReader.GetSafe(ParticleIndex, DefaultEnabled) && (VisTagReader.GetSafe(ParticleIndex, DefaultVisTag) == Properties->RendererVisibility);
+			};
 
-		// Ensure the final list only contains particles that are alive and enabled
-		ParticlesWithComponents.Reserve(UsedSlots.Num());
-		for (uint32 ParticleIndex = 0; ParticleIndex < ParticleData.GetNumInstances(); ParticleIndex++)
+		TMap<int32, int32> ParticlesWithComponents;
+		TArray<int32> FreeList;
+		if (Properties->bAssignComponentsOnParticleID && ComponentPool.Num() > 0)
 		{
-			int32 ParticleID = UniqueIDReader.GetSafe(ParticleIndex, DefaultUniqueID);
-			int32 PoolIndex;
-			if (UsedSlots.RemoveAndCopyValue(ParticleID, PoolIndex))
+			FreeList.Reserve(ComponentPool.Num());
+
+			// Determine the slots that were assigned to particles last frame
+			TMap<int32, int32> UsedSlots;
+			UsedSlots.Reserve(ComponentPool.Num());
+			for (int32 EntryIndex = 0; EntryIndex < ComponentPool.Num(); ++EntryIndex)
 			{
-				if (bIsRendererEnabled && IsParticleEnabled(ParticleIndex))
+				FComponentPoolEntry& Entry = ComponentPool[EntryIndex];
+				if (Entry.LastAssignedToParticleID >= 0)
 				{
-					ParticlesWithComponents.Emplace(ParticleID, PoolIndex);
+					UsedSlots.Emplace(Entry.LastAssignedToParticleID, EntryIndex);
 				}
 				else
 				{
-					// Particle has disabled components since last tick, ensure the component for this entry gets deactivated before re-use
-					USceneComponent* Component = ComponentPool[PoolIndex].Component.Get();
-					if (Component && Component->IsActive())
+					FreeList.Add(EntryIndex);
+				}
+			}
+
+			// Ensure the final list only contains particles that are alive and enabled
+			ParticlesWithComponents.Reserve(UsedSlots.Num());
+			for (uint32 ParticleIndex = 0; ParticleIndex < ParticleData.GetNumInstances(); ParticleIndex++)
+			{
+				int32 ParticleID = UniqueIDReader.GetSafe(ParticleIndex, DefaultUniqueID);
+				int32 PoolIndex;
+				if (UsedSlots.RemoveAndCopyValue(ParticleID, PoolIndex))
+				{
+					if (bIsRendererEnabled && IsParticleEnabled(ParticleIndex))
 					{
-						Component->Deactivate();
-						Component->SetVisibility(false, true);
+						ParticlesWithComponents.Emplace(ParticleID, PoolIndex);
+					}
+					else
+					{
+						// Particle has disabled components since last tick, ensure the component for this entry gets deactivated before re-use
+						USceneComponent* Component = ComponentPool[PoolIndex].Component.Get();
+						if (Component && Component->IsActive())
+						{
+							Component->Deactivate();
+							Component->SetVisibility(false, true);
+						}
+						FreeList.Add(PoolIndex);
+						ComponentPool[PoolIndex].LastAssignedToParticleID = -1;
 					}
-					FreeList.Add(PoolIndex);
-					ComponentPool[PoolIndex].LastAssignedToParticleID = -1;
 				}
 			}
-		}
 
-		// Any remaining in the used slots are now free to be reclaimed, due to their particles either dying or having their component disabled
-		for (TPair<int32, int32> UsedSlot : UsedSlots)
-		{
-			// Particle has died since last tick, ensure the component for this entry gets deactivated before re-use
-			USceneComponent* Component = ComponentPool[UsedSlot.Value].Component.Get();
-			if (Component && Component->IsActive())
+			// Any remaining in the used slots are now free to be reclaimed, due to their particles either dying or having their component disabled
+			for (TPair<int32, int32> UsedSlot : UsedSlots)
 			{
-				Component->Deactivate();
-				Component->SetVisibility(false, true);
+				// Particle has died since last tick, ensure the component for this entry gets deactivated before re-use
+				USceneComponent* Component = ComponentPool[UsedSlot.Value].Component.Get();
+				if (Component && Component->IsActive())
+				{
+					Component->Deactivate();
+					Component->SetVisibility(false, true);
+				}
+				FreeList.Add(UsedSlot.Value);
+				ComponentPool[UsedSlot.Value].LastAssignedToParticleID = -1;
 			}
-			FreeList.Add(UsedSlot.Value);
-			ComponentPool[UsedSlot.Value].LastAssignedToParticleID = -1;
-		}
-	}
-
-	const int32 MaxComponents = Properties->ComponentCountLimit;
-	int32 ComponentCount = 0;
-	for (uint32 ParticleIndex = 0; ParticleIndex < ParticleData.GetNumInstances(); ParticleIndex++)
-	{
-		if (!bIsRendererEnabled || !IsParticleEnabled(ParticleIndex))
-		{
-			// Skip particles that don't want a component
-			continue;
-		}
-
-		int32 ParticleID = -1;
-		int32 PoolIndex = -1;
-		if (Properties->bAssignComponentsOnParticleID)
-		{
-			// Get the particle ID and see if we have any components already assigned to the particle
-			ParticleID = UniqueIDReader.GetSafe(ParticleIndex, DefaultUniqueID);
-			ParticlesWithComponents.RemoveAndCopyValue(ParticleID, PoolIndex);
 		}
 
-		if (PoolIndex == -1 && ComponentCount + ParticlesWithComponents.Num() >= MaxComponents)
-		{
-			// The pool is full and there aren't any unused slots to claim
-			continue;
-		}
+		const int32 MaxComponents = Properties->ComponentCountLimit;
+		int32 ComponentCount = 0;
 
-		// Acquire a component for this particle
-		UGeometryCacheComponent* GeometryComponent = nullptr;
-		if (PoolIndex == -1)
+		for (uint32 ParticleIndex = 0; ParticleIndex < ParticleData.GetNumInstances(); ParticleIndex++)
 		{
-			// Start by trying to pull from the pool
-			if (!Properties->bAssignComponentsOnParticleID)
+			if (!bIsRendererEnabled || !IsParticleEnabled(ParticleIndex))
 			{
-				// We can just take the next slot
-				PoolIndex = ComponentCount < ComponentPool.Num() ? ComponentCount : -1;
+				// Skip particles that don't want a component
+				continue;
 			}
-			else if (FreeList.Num())
+
+			int32 ParticleID = -1;
+			int32 PoolIndex = -1;
+			if (Properties->bAssignComponentsOnParticleID)
 			{
-				PoolIndex = FreeList.Pop(false);
+				// Get the particle ID and see if we have any components already assigned to the particle
+				ParticleID = UniqueIDReader.GetSafe(ParticleIndex, DefaultUniqueID);
+				ParticlesWithComponents.RemoveAndCopyValue(ParticleID, PoolIndex);
 			}
-		}
-
-		if (PoolIndex >= 0)
-		{
-			GeometryComponent = ComponentPool[PoolIndex].Component.Get();
-		}
 
-		bool bCreateNewComponent = !GeometryComponent || GeometryComponent->HasAnyFlags(RF_BeginDestroyed | RF_FinishDestroyed);
-		int32 CacheIndex = GetGeometryCacheIndex(ArrayIndexReader, bCreateNewComponent, Emitter, Properties, ParticleIndex, ParticleID);
-		if (Properties->GeometryCaches.IsValidIndex(CacheIndex))
-		{
-			const FNiagaraGeometryCacheReference& Entry = Properties->GeometryCaches[CacheIndex];
-			UGeometryCache* Cache = ResolveGeometryCache(Entry, Emitter);
-			if (!Cache)
+			if (PoolIndex == -1 && ComponentCount + ParticlesWithComponents.Num() >= MaxComponents)
 			{
-				// no valid geometry cache was provided for this particle
+				// The pool is full and there aren't any unused slots to claim
 				continue;
 			}
-			
-			if (bCreateNewComponent)
-			{
-				// Determine the owner actor or spawn one
-				AActor* OwnerActor = SpawnedOwner.Get();
-				if (OwnerActor == nullptr)
-				{
-					OwnerActor = AttachComponent->GetOwner();
-					if (OwnerActor == nullptr)
-					{
-						// NOTE: This can happen with spawned systems
-						OwnerActor = AttachComponent->GetWorld()->SpawnActor<AActor>();
-						OwnerActor->SetFlags(RF_Transient);
-						SpawnedOwner = OwnerActor;
-					}
-				}
 
-				// if we don't have a pooled component we create a new one from the template
-				GeometryComponent = NewObject<UGeometryCacheComponent>(OwnerActor);
-				IGeometryCacheStreamingManager::Get().AddStreamingComponent(GeometryComponent);
-				GeometryComponent->SetFlags(RF_Transient);
-				GeometryComponent->SetupAttachment(AttachComponent);
-				GeometryComponent->RegisterComponent();
-				GeometryComponent->AddTickPrerequisiteComponent(AttachComponent);
-				GeometryComponent->SetLooping(Properties->bIsLooping);
-				GeometryComponent->SetManualTick(true); // we want to tick the component with the delta time of the niagara sim
-
-				if (Emitter->GetCachedEmitterData()->bLocalSpace)
+			// Acquire a component for this particle
+			if (PoolIndex == -1)
+			{
+				// Start by trying to pull from the pool
+				if (!Properties->bAssignComponentsOnParticleID)
 				{
-					GeometryComponent->SetAbsolute(false, false, false);
+					// We can just take the next slot
+					PoolIndex = ComponentCount < ComponentPool.Num() ? ComponentCount : -1;
 				}
-				else
+				else if (FreeList.Num())
 				{
-					GeometryComponent->SetAbsolute(true, true, true);
+					PoolIndex = FreeList.Pop(false);
 				}
+			}
 
-				if (PoolIndex >= 0)
-				{
-					// This should only happen if the component was destroyed externally
-					ComponentPool[PoolIndex].Component = GeometryComponent;
-					ComponentPool[PoolIndex].LastElapsedTime = 0;
-				}
-				else
-				{
-					// Add a new pool entry
-					PoolIndex = ComponentPool.Num();
-					ComponentPool.AddDefaulted_GetRef().Component = GeometryComponent;
-					ComponentPool[PoolIndex].LastElapsedTime = 0;
-				}
+			TOptional<int32> DefaultCacheIndex;
+			if (ArrayIndexReader.IsValid())
+			{
+				DefaultCacheIndex = ArrayIndexReader.GetSafe(ParticleIndex, INDEX_NONE);
 			}
 
-			// set the cache and override materials, since they might have changed from the last time we accessed the array index
-			// if they did not change from last tick these calls should be no-ops.
-			GeometryComponent->SetGeometryCache(Cache);
-			for (int32 i = 0; i < Entry.OverrideMaterials.Num(); i++)
+			UGeometryCacheComponent* GeometryComponent = CreateOrGetPooledComponent(Properties, AttachComponent, Emitter, DefaultCacheIndex, ParticleIndex, ParticleID, PoolIndex);
+			if (GeometryComponent == nullptr)
 			{
-				const TObjectPtr<UMaterialInterface>& OverrideMat = Entry.OverrideMaterials[i];
-				GeometryComponent->SetMaterial(i, OverrideMat);
+				continue;
 			}
-		}
-		else if (bCreateNewComponent)
-		{
-			// can't create a new component without valid geometry cache
-			continue;
-		}
 
-		const FNiagaraPosition DefaultPosition = ParameterStore.GetParameterValueOrDefault(Properties->PositionBinding.GetParamMapBindableVariable(), FNiagaraPosition(ForceInit));
-		const FVector3f DefaultRotation = ParameterStore.GetParameterValueOrDefault(Properties->RotationBinding.GetParamMapBindableVariable(), FVector3f::ZeroVector);
-		const FVector3f DefaultScale = ParameterStore.GetParameterValueOrDefault(Properties->ScaleBinding.GetParamMapBindableVariable(), FVector3f::OneVector);
-		const float DefaultElapsedTime = ParameterStore.GetParameterValueOrDefault(Properties->ElapsedTimeBinding.GetParamMapBindableVariable(), 0.0f);
-
-		const auto PositionAccessor = Properties->PositionAccessor.GetReader(Data);
-		const auto RotationAccessor = Properties->RotationAccessor.GetReader(Data);
-		const auto ScaleAccessor = Properties->ScaleAccessor.GetReader(Data);
-		const auto ElapsedTimeAccessor = Properties->ElapsedTimeAccessor.GetReader(Data);
-		const FNiagaraLWCConverter LwcConverter = SystemInstance->GetLWCConverter(Emitter->GetCachedEmitterData()->bLocalSpace);
-
-		FVector Position = LwcConverter.ConvertSimulationPositionToWorld(PositionAccessor.GetSafe(ParticleIndex, DefaultPosition));
-		FVector Scale = (FVector)ScaleAccessor.GetSafe(ParticleIndex, DefaultScale);
-		FVector RotationVector = (FVector)RotationAccessor.GetSafe(ParticleIndex, DefaultRotation);
-		FTransform Transform(FRotator(RotationVector.X, RotationVector.Y, RotationVector.Z), Position, Scale);
-		GeometryComponent->SetRelativeTransform(Transform);
-
-		// activate the component
-		if (!GeometryComponent->IsActive())
-		{
-			GeometryComponent->SetVisibility(true, true);
-			GeometryComponent->Activate(false);
-		}
+			const auto PositionAccessor = Properties->PositionAccessor.GetReader(Data);
+			const auto RotationAccessor = Properties->RotationAccessor.GetReader(Data);
+			const auto ScaleAccessor = Properties->ScaleAccessor.GetReader(Data);
+			const auto ElapsedTimeAccessor = Properties->ElapsedTimeAccessor.GetReader(Data);
 
-		float ElapsedTime = ElapsedTimeAccessor.GetSafe(ParticleIndex, DefaultElapsedTime);
+			FVector Position = LwcConverter.ConvertSimulationPositionToWorld(PositionAccessor.GetSafe(ParticleIndex, DefaultPosition));
+			FVector Scale = (FVector)ScaleAccessor.GetSafe(ParticleIndex, DefaultScale);
+			FVector RotationVector = (FVector)RotationAccessor.GetSafe(ParticleIndex, DefaultRotation);
+			FTransform Transform(FRotator(RotationVector.X, RotationVector.Y, RotationVector.Z), Position, Scale);
+			GeometryComponent->SetRelativeTransform(Transform);
 
-		float Duration = GeometryComponent->GetDuration();
-		if (Properties->bIsLooping && ElapsedTime < 0)
-		{
-			ElapsedTime -= FMath::CeilToFloat(ElapsedTime / Duration) * Duration;
-		}
-		else if (Properties->bIsLooping && ElapsedTime > Duration)
-		{
-			ElapsedTime -= FMath::Floor(ElapsedTime / Duration) * Duration;
-		}
+			// activate the component
+			if (!GeometryComponent->IsActive())
+			{
+				GeometryComponent->SetVisibility(true, true);
+				GeometryComponent->Activate(false);
+			}
 
-		FComponentPoolEntry& PoolEntry = ComponentPool[PoolIndex];
-		GeometryComponent->TickAtThisTime(ElapsedTime, true, ElapsedTime < PoolEntry.LastElapsedTime, Properties->bIsLooping);
+			float ElapsedTime = ElapsedTimeAccessor.GetSafe(ParticleIndex, DefaultElapsedTime);
 
-		PoolEntry.LastElapsedTime = ElapsedTime;
-		PoolEntry.LastAssignedToParticleID = ParticleID;
-		PoolEntry.LastActiveTime = CurrentTime;
+			float Duration = GeometryComponent->GetDuration();
+			if (Properties->bIsLooping && ElapsedTime < 0)
+			{
+				ElapsedTime -= FMath::CeilToFloat(ElapsedTime / Duration) * Duration;
+			}
+			else if (Properties->bIsLooping && ElapsedTime > Duration)
+			{
+				ElapsedTime -= FMath::Floor(ElapsedTime / Duration) * Duration;
+			}
 
-		++ComponentCount;
-		if (ComponentCount >= MaxComponents)
-		{
-			// We've hit our prescribed limit
-			break;
+			FComponentPoolEntry& PoolEntry = ComponentPool[PoolIndex];
+			GeometryComponent->TickAtThisTime(ElapsedTime, true, ElapsedTime < PoolEntry.LastElapsedTime, Properties->bIsLooping);
+
+			PoolEntry.LastElapsedTime = ElapsedTime;
+			PoolEntry.LastAssignedToParticleID = ParticleID;
+			PoolEntry.LastActiveTime = CurrentTime;
+
+			++ComponentCount;
+			if (ComponentCount >= MaxComponents)
+			{
+				// We've hit our prescribed limit
+				break;
+			}
 		}
-	}
 
-	if (ComponentCount < ComponentPool.Num())
-	{
-		// go over the pooled components we didn't need this tick to see if we can destroy some and deactivate the rest
-		for (int32 PoolIndex = 0; PoolIndex < ComponentPool.Num(); ++PoolIndex)
+		// Go over the pooled components we didn't need this tick to see if we can destroy some and deactivate the rest
+		if (ComponentCount < ComponentPool.Num())
 		{
-			FComponentPoolEntry& PoolEntry = ComponentPool[PoolIndex];
-			if (Properties->bAssignComponentsOnParticleID)
+			for (int32 PoolIndex = 0; PoolIndex < ComponentPool.Num(); ++PoolIndex)
 			{
-				if (PoolEntry.LastAssignedToParticleID >= 0)
+				FComponentPoolEntry& PoolEntry = ComponentPool[PoolIndex];
+				if (Properties->bAssignComponentsOnParticleID)
+				{
+					if (PoolEntry.LastAssignedToParticleID >= 0)
+					{
+						// This one's in use
+						continue;
+					}
+				}
+				else if (PoolIndex < ComponentCount)
 				{
-					// This one's in use
 					continue;
 				}
+
+				USceneComponent* Component = PoolEntry.Component.Get();
+				if (!Component || (CurrentTime - PoolEntry.LastActiveTime) >= GNiagaraGeometryComponentRenderPoolInactiveTimeLimit)
+				{
+					if (Component)
+					{
+						IGeometryCacheStreamingManager::Get().RemoveStreamingComponent(PoolEntry.Component.Get());
+						Component->DestroyComponent();
+					}
+
+					// destroy the component pool slot
+					ComponentPool.RemoveAtSwap(PoolIndex, 1, false);
+					--PoolIndex;
+					continue;
+				}
+				else if (Component->IsActive())
+				{
+					Component->Deactivate();
+					Component->SetVisibility(false, true);
+				}
 			}
-			else if (PoolIndex < ComponentCount)
+		}
+	}
+	// Emitter mode we can skip this if we aren't enabled / visible
+	else if (bIsRendererEnabled && DefaultVisTag == Properties->RendererVisibility)
+	{
+		int32 PoolIndex = ComponentPool.Num() > 0 ? 0 : -1;
+		int32 UsedComponentCount = 0;
+		TOptional<int32> DefaultCacheIndex;
+		DefaultCacheIndex = 0;
+		if ( UGeometryCacheComponent* GeometryComponent = CreateOrGetPooledComponent(Properties, AttachComponent, Emitter, DefaultCacheIndex, 0, 0, PoolIndex))
+		{
+			const FVector Position = LwcConverter.ConvertSimulationPositionToWorld(DefaultPosition);
+			const FVector Scale = FVector(DefaultScale);
+			const FVector RotationVector = FVector(DefaultRotation);
+			const FTransform Transform(FRotator(RotationVector.X, RotationVector.Y, RotationVector.Z), Position, Scale);
+			GeometryComponent->SetRelativeTransform(Transform);
+
+			// activate the component
+			if (!GeometryComponent->IsActive())
 			{
-				continue;
+				GeometryComponent->SetVisibility(true, true);
+				GeometryComponent->Activate(false);
 			}
 
-			USceneComponent* Component = PoolEntry.Component.Get();
-			if (!Component || (CurrentTime - PoolEntry.LastActiveTime) >= GNiagaraGeometryComponentRenderPoolInactiveTimeLimit)
+			const float Duration = GeometryComponent->GetDuration();
+			float ElapsedTime = DefaultElapsedTime;
+			if (Properties->bIsLooping && ElapsedTime < 0)
 			{
-				if (Component)
-				{
-					IGeometryCacheStreamingManager::Get().RemoveStreamingComponent(PoolEntry.Component.Get());
-					Component->DestroyComponent();
-				}
-
-				// destroy the component pool slot
-				ComponentPool.RemoveAtSwap(PoolIndex, 1, false);
-				--PoolIndex;
-				continue;
+				ElapsedTime -= FMath::CeilToFloat(ElapsedTime / Duration) * Duration;
+			}
+			else if (Properties->bIsLooping && ElapsedTime > Duration)
+			{
+				ElapsedTime -= FMath::Floor(ElapsedTime / Duration) * Duration;
 			}
-			else if (Component->IsActive())
+
+			FComponentPoolEntry& PoolEntry = ComponentPool[PoolIndex];
+			GeometryComponent->TickAtThisTime(ElapsedTime, true, ElapsedTime < PoolEntry.LastElapsedTime, Properties->bIsLooping);
+
+			PoolEntry.LastElapsedTime = ElapsedTime;
+			PoolEntry.LastAssignedToParticleID = INDEX_NONE;
+			PoolEntry.LastActiveTime = CurrentTime;
+
+			++UsedComponentCount;
+		}
+
+		// We just scrub all the components we don't use, if we were to turn on / off the components we may want to reconsider this
+		while ( ComponentPool.Num() > UsedComponentCount )
+		{
+			FComponentPoolEntry& PoolEntry = ComponentPool.Last();
+			if ( USceneComponent* Component = PoolEntry.Component.Get() )
 			{
-				Component->Deactivate();
-				Component->SetVisibility(false, true);
+				IGeometryCacheStreamingManager::Get().RemoveStreamingComponent(PoolEntry.Component.Get());
+				Component->DestroyComponent();
 			}
 		}
 	}
