GeometryCache: Expose streamer settings to Python so that we can test things during streaming from automated tests.

diff --git a/Engine/Plugins/Runtime/GeometryCache/Source/GeometryCacheStreamer/Public/GeometryCacheStreamerSettings.h b/Engine/Plugins/Runtime/GeometryCache/Source/GeometryCacheStreamer/Public/GeometryCacheStreamerSettings.h
--- a/Engine/Plugins/Runtime/GeometryCache/Source/GeometryCacheStreamer/Public/GeometryCacheStreamerSettings.h
+++ b/Engine/Plugins/Runtime/GeometryCache/Source/GeometryCacheStreamer/Public/GeometryCacheStreamerSettings.h
@@ -17,10 +17,10 @@ public:
 	UGeometryCacheStreamerSettings();
 
 	/** The amount of animation (in seconds) to stream ahead of time (per stream) */
-	UPROPERTY(config, EditAnywhere, Category = "Geometry Cache Streamer", meta = (DisplayName = "Look-Ahead Buffer (in seconds)", ClampMin = "0.01", ClampMax = "3600.0"))
+	UPROPERTY(config, EditAnywhere, BlueprintReadWrite, Category = "Geometry Cache Streamer", meta = (DisplayName = "Look-Ahead Buffer (in seconds)", ClampMin = "0.01", ClampMax = "3600.0"))
 	float LookAheadBuffer;
 
 	/** The maximum total amount of streamed data allowed in memory (for all streams) */
-	UPROPERTY(config, EditAnywhere, Category = "Geometry Cache Streamer", meta = (DisplayName = "Maximum Memory Allowed (in MB)", ClampMin = "1.0", ClampMax = "262144.0"))
+	UPROPERTY(config, EditAnywhere, BlueprintReadWrite, Category = "Geometry Cache Streamer", meta = (DisplayName = "Maximum Memory Allowed (in MB)", ClampMin = "1.0", ClampMax = "262144.0"))
 	float MaxMemoryAllowed;
 };
