- Add MIC support to Geometry Cache Renderer

diff --git a/Engine/Plugins/Runtime/GeometryCache/Source/GeometryCache/Classes/NiagaraGeometryCacheRendererProperties.h b/Engine/Plugins/Runtime/GeometryCache/Source/GeometryCache/Classes/NiagaraGeometryCacheRendererProperties.h
--- a/Engine/Plugins/Runtime/GeometryCache/Source/GeometryCache/Classes/NiagaraGeometryCacheRendererProperties.h
+++ b/Engine/Plugins/Runtime/GeometryCache/Source/GeometryCache/Classes/NiagaraGeometryCacheRendererProperties.h
@@ -14,13 +14,27 @@ struct FVersionedNiagaraEmitter;
 class FNiagaraEmitterInstance;
 class SWidget;
 
+USTRUCT()
+struct FNiagaraGeometryCacheMICOverride
+{
+	GENERATED_USTRUCT_BODY()
+
+	UPROPERTY()
+	TObjectPtr<UMaterialInterface> OriginalMaterial;
+
+	UPROPERTY()
+	TObjectPtr<UMaterialInstanceConstant> ReplacementMaterial;
+};
+
 USTRUCT()
 struct FNiagaraGeometryCacheReference
 {
 	GENERATED_USTRUCT_BODY()
-public:
+
 	FNiagaraGeometryCacheReference();
-	
+
+	UMaterialInterface* ResolveMaterial(UGeometryCache* ResolvedCache, int32 MaterialIndex) const;
+
 	/** Reference to the geometry cache asset to use (if the user parameter binding is not set) */
 	UPROPERTY(EditAnywhere, Category = "GeometryCache")
 	TObjectPtr<UGeometryCache> GeometryCache;
@@ -33,11 +47,8 @@ public:
 	UPROPERTY(EditAnywhere, Category = "GeometryCache")
 	TArray<TObjectPtr<UMaterialInterface>> OverrideMaterials;
 
-//-TODO: Follow what meshes do
-//#if WITH_EDITORONLY_DATA
-//	UPROPERTY(transient)
-//	TArray<TObjectPtr<UMaterialInstanceConstant>> MICMaterials;
-//#endif
+	UPROPERTY()
+	TArray<FNiagaraGeometryCacheMICOverride> MICOverrideMaterials;
 };
 
 UCLASS(editinlinenew, MinimalAPI, meta = (DisplayName = "Geometry Cache Renderer"))
@@ -51,7 +62,6 @@ public:
 	//UObject Interface
 	virtual void PostLoad() override;
 	virtual void PostInitProperties() override;
-	//virtual void Serialize(FArchive& Ar) override;
 #if WITH_EDITORONLY_DATA
 	virtual void PostEditChangeProperty(struct FPropertyChangedEvent& PropertyChangedEvent) override;
 #endif// WITH_EDITORONLY_DATA
diff --git a/Engine/Plugins/Runtime/GeometryCache/Source/GeometryCache/Private/NiagaraGeometryCacheRendererProperties.cpp b/Engine/Plugins/Runtime/GeometryCache/Source/GeometryCache/Private/NiagaraGeometryCacheRendererProperties.cpp
--- a/Engine/Plugins/Runtime/GeometryCache/Source/GeometryCache/Private/NiagaraGeometryCacheRendererProperties.cpp
+++ b/Engine/Plugins/Runtime/GeometryCache/Source/GeometryCache/Private/NiagaraGeometryCacheRendererProperties.cpp
@@ -1,15 +1,17 @@
 // Copyright Epic Games, Inc. All Rights Reserved.
 
 #include "NiagaraGeometryCacheRendererProperties.h"
-#include "GeometryCache.h"
 #include "NiagaraConstants.h"
 #include "NiagaraEmitter.h"
 #include "NiagaraEmitterInstance.h"
 #include "NiagaraModule.h"
+#include "NiagaraParameterStore.h"
 #include "NiagaraRendererGeometryCache.h"
+
 #include "Modules/ModuleManager.h"
-#include "NiagaraParameterStore.h"
+#include "Materials/MaterialInstanceConstant.h"
 #include "Widgets/Text/STextBlock.h"
+#include "GeometryCache.h"
 
 #include UE_INLINE_GENERATED_CPP_BY_NAME(NiagaraGeometryCacheRendererProperties)
 
@@ -33,6 +35,25 @@ FNiagaraGeometryCacheReference::FNiagaraGeometryCacheReference()
 {
 }
 
+UMaterialInterface* FNiagaraGeometryCacheReference::ResolveMaterial(UGeometryCache* ResolvedCache, int32 MaterialIndex) const
+{
+	UMaterialInterface* Material = OverrideMaterials.IsValidIndex(MaterialIndex) ? OverrideMaterials[MaterialIndex] : nullptr;
+	if (Material == nullptr && GeometryCache->Materials.IsValidIndex(MaterialIndex))
+	{
+		Material = GeometryCache->Materials[MaterialIndex];
+	}
+
+	for (const FNiagaraGeometryCacheMICOverride& MICOverride : MICOverrideMaterials)
+	{
+		if (MICOverride.OriginalMaterial == Material)
+		{
+			return MICOverride.ReplacementMaterial;
+		}
+	}
+
+	return Material;
+}
+
 UNiagaraGeometryCacheRendererProperties::UNiagaraGeometryCacheRendererProperties()
 {
 	AttributeBindings.Reserve(6);
@@ -100,10 +121,6 @@ void UNiagaraGeometryCacheRendererProperties::InitDefaultAttributes()
 	}
 }
 
-//void UNiagaraGeometryCacheRendererProperties::Serialize(FArchive& Ar)
-//{
-//}
-
 void UNiagaraGeometryCacheRendererProperties::PostInitProperties()
 {
 	Super::PostInitProperties();
@@ -166,11 +183,11 @@ void UNiagaraGeometryCacheRendererProperties::PostEditChangeProperty(struct FPro
 
 	// Update our MICs if we change material / material bindings
 	//-OPT: Could narrow down further to only static materials
-	//if ((PropertyName == GET_MEMBER_NAME_CHECKED(UNiagaraGeometryCacheRendererProperties, Material)) ||
-	//	(MemberPropertyName == GET_MEMBER_NAME_CHECKED(UNiagaraGeometryCacheRendererProperties, MaterialParameters)))
-	//{
-	//	UpdateMICs();
-	//}
+	if ((MemberPropertyName == GET_MEMBER_NAME_CHECKED(UNiagaraGeometryCacheRendererProperties, GeometryCaches)) ||
+		(MemberPropertyName == GET_MEMBER_NAME_CHECKED(UNiagaraGeometryCacheRendererProperties, MaterialParameters)))
+	{
+		UpdateMICs();
+	}
 }
 #endif// WITH_EDITORONLY_DATA
 
@@ -234,22 +251,18 @@ void UNiagaraGeometryCacheRendererProperties::GetUsedMaterials(const FNiagaraEmi
 {
 	for (const FNiagaraGeometryCacheReference& Entry : GeometryCaches)
 	{
-		UGeometryCache* GeometryCache = ResolveGeometryCache(Entry, EmitterInstance);;
+		UGeometryCache* GeometryCache = ResolveGeometryCache(Entry, EmitterInstance);
 		if (GeometryCache == nullptr)
 		{
 			continue;
 		}
 
 		const int32 MaxIndex = FMath::Max(GeometryCache->Materials.Num(), Entry.OverrideMaterials.Num());
-		for (int i = 0; i < MaxIndex; i++)
+		for (int i=0; i < MaxIndex; ++i)
 		{
-			if (Entry.OverrideMaterials.IsValidIndex(i) && Entry.OverrideMaterials[i])
+			if (UMaterialInterface* Material = Entry.ResolveMaterial(GeometryCache, i))
 			{
-				MaterialInterfaces.AddUnique(Entry.OverrideMaterials[i]);
-			}
-			else if (GeometryCache->Materials.IsValidIndex(i))
-			{
-				MaterialInterfaces.AddUnique(GeometryCache->Materials[i]);
+				MaterialInterfaces.AddUnique(Material);
 			}
 		}
 	}
@@ -447,7 +460,58 @@ void UNiagaraGeometryCacheRendererProperties::RemoveVariable(const FNiagaraVaria
 void UNiagaraGeometryCacheRendererProperties::UpdateMICs()
 {
 #if WITH_EDITORONLY_DATA
-	//UpdateMaterialParametersMIC(MaterialParameters, Material, MICMaterial);
+	// Grab existing MICs so we can reuse and clear them out so they aren't applied during GetUsedMaterials
+	TArray<TObjectPtr<UMaterialInstanceConstant>> MICMaterials;
+	for (FNiagaraGeometryCacheReference& GeoCache : GeometryCaches)
+	{
+		for (const FNiagaraGeometryCacheMICOverride& MICOverride : GeoCache.MICOverrideMaterials)
+		{
+			MICMaterials.AddUnique(MICOverride.ReplacementMaterial);
+		}
+		GeoCache.MICOverrideMaterials.Reset();
+	}
+
+	// Gather materials and generate MICs
+	TArray<UMaterialInterface*> Materials;
+	GetUsedMaterials(nullptr, Materials);
+
+	UpdateMaterialParametersMIC(MaterialParameters, Materials, MICMaterials);
+
+	// Push MIC overrides back into each geometry cache
+	if (MICMaterials.Num() > 0)
+	{
+		for (FNiagaraGeometryCacheReference& GeoCache : GeometryCaches)
+		{
+			UGeometryCache* GeometryCache = ResolveGeometryCache(GeoCache, nullptr);
+			if (GeometryCache == nullptr)
+			{
+				continue;
+			}
+
+			const int32 MaxIndex = FMath::Max(GeometryCache->Materials.Num(), GeoCache.OverrideMaterials.Num());
+			for (int i=0; i < MaxIndex; ++i)
+			{
+				UMaterialInterface* OriginalMaterial = GeoCache.OverrideMaterials.IsValidIndex(i) ? GeoCache.OverrideMaterials[i] : nullptr;
+				if (OriginalMaterial == nullptr && GeometryCache->Materials.IsValidIndex(i))
+				{
+					OriginalMaterial = GeometryCache->Materials[i];
+				}
+				if (OriginalMaterial == nullptr)
+				{
+					continue;
+				}
+
+				const int32 MICIndex = Materials.IndexOfByKey(OriginalMaterial);
+				UMaterialInstanceConstant* MIC = MICMaterials.IsValidIndex(MICIndex) ? MICMaterials[MICIndex] : nullptr;
+				if (MIC)
+				{
+					FNiagaraGeometryCacheMICOverride& MICOverride = GeoCache.MICOverrideMaterials.AddDefaulted_GetRef();
+					MICOverride.OriginalMaterial = OriginalMaterial;
+					MICOverride.ReplacementMaterial = MIC;
+				}
+			}
+		}
+	}
 #endif
 }
 
diff --git a/Engine/Plugins/Runtime/GeometryCache/Source/GeometryCache/Private/NiagaraRendererGeometryCache.cpp b/Engine/Plugins/Runtime/GeometryCache/Source/GeometryCache/Private/NiagaraRendererGeometryCache.cpp
--- a/Engine/Plugins/Runtime/GeometryCache/Source/GeometryCache/Private/NiagaraRendererGeometryCache.cpp
+++ b/Engine/Plugins/Runtime/GeometryCache/Source/GeometryCache/Private/NiagaraRendererGeometryCache.cpp
@@ -58,7 +58,7 @@ void FNiagaraRendererGeometryCache::Initialize(const UNiagaraRendererProperties*
 
 		for ( int32 iMaterial=0; iMaterial < GeoCache->Materials.Num(); ++iMaterial)
 		{
-			UMaterialInterface* Material = GeoCacheRef.OverrideMaterials.IsValidIndex(iMaterial) ? GeoCacheRef.OverrideMaterials[iMaterial] : GeoCache->Materials[iMaterial];
+			UMaterialInterface* Material = GeoCacheRef.ResolveMaterial(GeoCache, iMaterial);
 
 			MaterialRemapTable[iCache][iMaterial] =
 				BaseMaterials_GT.IndexOfByPredicate(
