- Add emitter mode support for Geometry Cache Array Index

#rnx

diff --git a/Engine/Plugins/Runtime/GeometryCache/Source/GeometryCache/Private/NiagaraRendererGeometryCache.cpp b/Engine/Plugins/Runtime/GeometryCache/Source/GeometryCache/Private/NiagaraRendererGeometryCache.cpp
--- a/Engine/Plugins/Runtime/GeometryCache/Source/GeometryCache/Private/NiagaraRendererGeometryCache.cpp
+++ b/Engine/Plugins/Runtime/GeometryCache/Source/GeometryCache/Private/NiagaraRendererGeometryCache.cpp
@@ -249,7 +249,7 @@ void FNiagaraRendererGeometryCache::PostSystemTick_GameThread(const UNiagaraRend
 	const FNiagaraParameterStore& ParameterStore = Emitter->GetRendererBoundVariables();
 	const FNiagaraBool DefaultEnabled = ParameterStore.GetParameterValueOrDefault(Properties->EnabledBinding.GetParamMapBindableVariable(), FNiagaraBool(true));
 	const int32 DefaultVisTag = ParameterStore.GetParameterValueOrDefault(Properties->RendererVisibilityTagBinding.GetParamMapBindableVariable(), Properties->RendererVisibility);
-	//const int32 DefaultArrayIndex = ParameterStore.GetParameterValueOrDefault(Properties->ArrayIndexBinding.GetParamMapBindableVariable(), INDEX_NONE);	//-TODO: Revisit this when emitter mode is supported?
+	const TOptional<int32> DefaultArrayIndex = ParameterStore.GetParameterOptionalValue<int32>(Properties->ArrayIndexBinding.GetParamMapBindableVariable());
 	const int32 DefaultUniqueID = INDEX_NONE;
 	const FNiagaraPosition DefaultPosition = ParameterStore.GetParameterValueOrDefault(Properties->PositionBinding.GetParamMapBindableVariable(), FNiagaraPosition(ForceInit));
 	const FVector3f DefaultRotation = ParameterStore.GetParameterValueOrDefault(Properties->RotationBinding.GetParamMapBindableVariable(), FVector3f::ZeroVector);
@@ -380,7 +380,7 @@ void FNiagaraRendererGeometryCache::PostSystemTick_GameThread(const UNiagaraRend
 				}
 			}
 
-			TOptional<int32> DefaultCacheIndex;
+			TOptional<int32> DefaultCacheIndex = DefaultArrayIndex;
 			if (ArrayIndexReader.IsValid())
 			{
 				DefaultCacheIndex = ArrayIndexReader.GetSafe(ParticleIndex, INDEX_NONE);
@@ -483,9 +483,7 @@ void FNiagaraRendererGeometryCache::PostSystemTick_GameThread(const UNiagaraRend
 	{
 		int32 PoolIndex = ComponentPool.Num() > 0 ? 0 : -1;
 		int32 UsedComponentCount = 0;
-		TOptional<int32> DefaultCacheIndex;
-		DefaultCacheIndex = 0;
-		if ( UGeometryCacheComponent* GeometryComponent = CreateOrGetPooledComponent(Properties, AttachComponent, Emitter, DefaultCacheIndex, 0, 0, PoolIndex))
+		if ( UGeometryCacheComponent* GeometryComponent = CreateOrGetPooledComponent(Properties, AttachComponent, Emitter, DefaultArrayIndex, 0, 0, PoolIndex))
 		{
 			const FVector Position = LwcConverter.ConvertSimulationPositionToWorld(DefaultPosition);
 			const FVector Scale = FVector(DefaultScale);
