GeometryCacheConstantTopologyWriter: fix signed/unsigned mismatch

diff --git a/Engine/Plugins/Runtime/GeometryCache/Source/GeometryCache/Private/GeometryCacheConstantTopologyWriter.cpp b/Engine/Plugins/Runtime/GeometryCache/Source/GeometryCache/Private/GeometryCacheConstantTopologyWriter.cpp
--- a/Engine/Plugins/Runtime/GeometryCache/Source/GeometryCache/Private/GeometryCacheConstantTopologyWriter.cpp
+++ b/Engine/Plugins/Runtime/GeometryCache/Source/GeometryCache/Private/GeometryCacheConstantTopologyWriter.cpp
@@ -265,7 +265,7 @@ namespace UE::GeometryCacheHelpers
 		{
 			return false;
 		}
-		const int32 NumVertices = PositionsToMoveFrom[0].Num();
+		const uint32 NumVertices = PositionsToMoveFrom[0].Num();
 		checkSlow(FMath::Max(Indices) < NumVertices);
 
 		TArray<FFrameData> Frames;
@@ -288,7 +288,7 @@ namespace UE::GeometryCacheHelpers
 		{
 			return false;
 		}
-		const int32 NumVertices = FramesToMoveFrom[0].Positions.Num();
+		const uint32 NumVertices = FramesToMoveFrom[0].Positions.Num();
 		checkSlow(FMath::Max(Indices) < NumVertices);
 		const bool bHasUV0 = !UVs.IsEmpty();
 		if (bHasUV0 && UVs.Num() != NumVertices)
