Fixed nonunity build break.

diff --git a/Engine/Plugins/Runtime/GeometryCache/Source/GeometryCache/Classes/GeometryCacheCodecBase.h b/Engine/Plugins/Runtime/GeometryCache/Source/GeometryCache/Classes/GeometryCacheCodecBase.h
--- a/Engine/Plugins/Runtime/GeometryCache/Source/GeometryCache/Classes/GeometryCacheCodecBase.h
+++ b/Engine/Plugins/Runtime/GeometryCache/Source/GeometryCache/Classes/GeometryCacheCodecBase.h
@@ -6,10 +6,10 @@
 #include "GeometryCacheCodecBase.generated.h"
 
 struct FResourceSizeEx;
-
 struct FGeometryCacheMeshData;
 class UGeometryCacheTrackStreamable;
 class FGeometryCacheTrackStreamableRenderResource;
+class FRHICommandListBase;
 
 /**
 The smallest unit of streamed geometry cache data.
