Reset cached scene velocity to prevent motion blur when reusing components

diff --git a/Engine/Plugins/FX/Niagara/Source/Niagara/Private/NiagaraRendererComponents.cpp b/Engine/Plugins/FX/Niagara/Source/Niagara/Private/NiagaraRendererComponents.cpp
--- a/Engine/Plugins/FX/Niagara/Source/Niagara/Private/NiagaraRendererComponents.cpp
+++ b/Engine/Plugins/FX/Niagara/Source/Niagara/Private/NiagaraRendererComponents.cpp
@@ -599,6 +599,11 @@ void FNiagaraRendererComponents::PostSystemTick_GameThread(const UNiagaraRendere
 			{
 				SceneComponent->SetVisibility(true, true);
 				SceneComponent->Activate(false);
+				if (UPrimitiveComponent* PrimitiveComponent = Cast<UPrimitiveComponent>(SceneComponent))
+				{
+					// prevent motion blur when reusing components
+					PrimitiveComponent->ResetSceneVelocity();
+				}
 			}
 		}
 
diff --git a/Engine/Plugins/Runtime/GeometryCache/Source/GeometryCache/Private/NiagaraRendererGeometryCache.cpp b/Engine/Plugins/Runtime/GeometryCache/Source/GeometryCache/Private/NiagaraRendererGeometryCache.cpp
--- a/Engine/Plugins/Runtime/GeometryCache/Source/GeometryCache/Private/NiagaraRendererGeometryCache.cpp
+++ b/Engine/Plugins/Runtime/GeometryCache/Source/GeometryCache/Private/NiagaraRendererGeometryCache.cpp
@@ -406,6 +406,7 @@ void FNiagaraRendererGeometryCache::PostSystemTick_GameThread(const UNiagaraRend
 			// activate the component
 			if (!GeometryComponent->IsActive())
 			{
+				GeometryComponent->ResetSceneVelocity();
 				GeometryComponent->SetVisibility(true, true);
 				GeometryComponent->Activate(false);
 			}
@@ -494,6 +495,7 @@ void FNiagaraRendererGeometryCache::PostSystemTick_GameThread(const UNiagaraRend
 			// activate the component
 			if (!GeometryComponent->IsActive())
 			{
+				GeometryComponent->ResetSceneVelocity();
 				GeometryComponent->SetVisibility(true, true);
 				GeometryComponent->Activate(false);
 			}
diff --git a/Engine/Source/Runtime/Engine/Classes/Components/PrimitiveComponent.h b/Engine/Source/Runtime/Engine/Classes/Components/PrimitiveComponent.h
--- a/Engine/Source/Runtime/Engine/Classes/Components/PrimitiveComponent.h
+++ b/Engine/Source/Runtime/Engine/Classes/Components/PrimitiveComponent.h
@@ -2544,6 +2544,9 @@ public:
 	ENGINE_API void SetSelectionOutlineColorIndex(uint8 SelectionOutlineColorIndex);
 #endif// WITH_EDITOR
 
+	/** Resets the cached scene velocity. Useful to prevent motion blur when teleporting components. See also SetIsBeingMovedByEditor(). */
+	ENGINE_API void ResetSceneVelocity();
+
 	/**
 	 *	Set the angular velocity of all bodies in this component.
 	 *
diff --git a/Engine/Source/Runtime/Engine/Private/Components/PrimitiveComponent.cpp b/Engine/Source/Runtime/Engine/Private/Components/PrimitiveComponent.cpp
--- a/Engine/Source/Runtime/Engine/Private/Components/PrimitiveComponent.cpp
+++ b/Engine/Source/Runtime/Engine/Private/Components/PrimitiveComponent.cpp
@@ -1894,8 +1894,17 @@ void UPrimitiveComponent::SetSelectionOutlineColorIndex(uint8 InSelectionOutline
 		SceneProxy->SetSelectionOutlineColorIndex_GameThread(InSelectionOutlineColorIndex);
 	}
 }
+
 #endif// WITH_EDITOR
 
+void UPrimitiveComponent::ResetSceneVelocity()
+{
+	if (SceneProxy)
+	{
+		SceneProxy->ResetSceneVelocity_GameThread();
+	}
+}
+
 void UPrimitiveComponent::PushHoveredToProxy(const bool bInHovered)
 {
 	//although this should only be called for attached components, some billboard components can get in without valid proxies
diff --git a/Engine/Source/Runtime/Engine/Private/PrimitiveSceneProxy.cpp b/Engine/Source/Runtime/Engine/Private/PrimitiveSceneProxy.cpp
--- a/Engine/Source/Runtime/Engine/Private/PrimitiveSceneProxy.cpp
+++ b/Engine/Source/Runtime/Engine/Private/PrimitiveSceneProxy.cpp
@@ -1356,8 +1356,21 @@ void FPrimitiveSceneProxy::SetSelectionOutlineColorIndex_GameThread(uint8 ColorI
 			SelectionOutlineColorIndex = ColorIndex;
 		});
 }
+
 #endif
 
+void FPrimitiveSceneProxy::ResetSceneVelocity_GameThread()
+{
+	check(IsInGameThread());
+
+	FPrimitiveSceneProxy* PrimitiveSceneProxy = this;
+	ENQUEUE_RENDER_COMMAND(ResetSceneVelocity)(
+		[PrimitiveSceneProxy](FRHICommandListImmediate& RHICmdList)
+		{
+			PrimitiveSceneProxy->GetScene().UpdatePrimitiveVelocityState_RenderThread(PrimitiveSceneProxy->GetPrimitiveSceneInfo(), false);
+		});
+}
+
 void FPrimitiveSceneProxy::SetEvaluateWorldPositionOffset_GameThread(bool bEvaluate)
 {
 	check(IsInGameThread());
diff --git a/Engine/Source/Runtime/Engine/Public/PrimitiveSceneProxy.h b/Engine/Source/Runtime/Engine/Public/PrimitiveSceneProxy.h
--- a/Engine/Source/Runtime/Engine/Public/PrimitiveSceneProxy.h
+++ b/Engine/Source/Runtime/Engine/Public/PrimitiveSceneProxy.h
@@ -242,6 +242,9 @@ public:
 	void SetSelectionOutlineColorIndex_GameThread(uint8 ColorIndex);
 #endif	// WITH_EDITOR
 
+	/** Enqueue and update for the render thread to remove the velocity data for this component from the scene. */
+	ENGINE_API void ResetSceneVelocity_GameThread();
+
 	/** Enqueue updated setting for evaluation of World Position Offset. */
 	void SetEvaluateWorldPositionOffset_GameThread(bool bEvaluate);
 
