Fixed missing non-unity includes

diff --git a/Engine/Plugins/FX/Niagara/Source/Niagara/Private/NiagaraRendererComponents.cpp b/Engine/Plugins/FX/Niagara/Source/Niagara/Private/NiagaraRendererComponents.cpp
--- a/Engine/Plugins/FX/Niagara/Source/Niagara/Private/NiagaraRendererComponents.cpp
+++ b/Engine/Plugins/FX/Niagara/Source/Niagara/Private/NiagaraRendererComponents.cpp
@@ -2,6 +2,7 @@
 
 #include "NiagaraRendererComponents.h"
 #include "Engine/World.h"
+#include "Components/PrimitiveComponent.h"
 #include "NiagaraDataSet.h"
 #include "NiagaraEmitterInstance.h"
 #include "NiagaraStats.h"
diff --git a/Engine/Plugins/Runtime/GeometryCache/Source/GeometryCache/Private/NiagaraRendererGeometryCache.cpp b/Engine/Plugins/Runtime/GeometryCache/Source/GeometryCache/Private/NiagaraRendererGeometryCache.cpp
--- a/Engine/Plugins/Runtime/GeometryCache/Source/GeometryCache/Private/NiagaraRendererGeometryCache.cpp
+++ b/Engine/Plugins/Runtime/GeometryCache/Source/GeometryCache/Private/NiagaraRendererGeometryCache.cpp
@@ -3,6 +3,7 @@
 #include "NiagaraRendererGeometryCache.h"
 
 #include "Engine/World.h"
+#include "Components/PrimitiveComponent.h"
 #include "GeometryCache.h"
 #include "GeometryCacheStreamingManager.h"
 #include "GeometryCacheComponent.h"
