[GeomCache] Added some virtual method to the tracks to sample at a specific frame directly (sample index).

diff --git a/Engine/Plugins/Experimental/GeometryCacheAbcFile/Source/GeometryCacheAbcFile/Private/GeometryCacheTrackAbcFile.cpp b/Engine/Plugins/Experimental/GeometryCacheAbcFile/Source/GeometryCacheAbcFile/Private/GeometryCacheTrackAbcFile.cpp
--- a/Engine/Plugins/Experimental/GeometryCacheAbcFile/Source/GeometryCacheAbcFile/Private/GeometryCacheTrackAbcFile.cpp
+++ b/Engine/Plugins/Experimental/GeometryCacheAbcFile/Source/GeometryCacheAbcFile/Private/GeometryCacheTrackAbcFile.cpp
@@ -330,6 +330,11 @@ bool UGeometryCacheTrackAbcFile::GetMeshDataAtTime(float Time, FGeometryCacheMes
 	return GetMeshData(SampleIndex, OutMeshData);
 }
 
+bool UGeometryCacheTrackAbcFile::GetMeshDataAtSampleIndex(int32 SampleIndex, FGeometryCacheMeshData& OutMeshData)
+{
+	return GetMeshData(SampleIndex, OutMeshData);
+}
+
 bool UGeometryCacheTrackAbcFile::GetMeshData(int32 SampleIndex, FGeometryCacheMeshData& OutMeshData)
 {
 	if (AbcFile)
diff --git a/Engine/Plugins/Experimental/GeometryCacheAbcFile/Source/GeometryCacheAbcFile/Public/GeometryCacheTrackAbcFile.h b/Engine/Plugins/Experimental/GeometryCacheAbcFile/Source/GeometryCacheAbcFile/Public/GeometryCacheTrackAbcFile.h
--- a/Engine/Plugins/Experimental/GeometryCacheAbcFile/Source/GeometryCacheAbcFile/Public/GeometryCacheTrackAbcFile.h
+++ b/Engine/Plugins/Experimental/GeometryCacheAbcFile/Source/GeometryCacheAbcFile/Public/GeometryCacheTrackAbcFile.h
@@ -27,6 +27,7 @@ public:
 	virtual const bool UpdateBoundsData(const float Time, const bool bLooping, const bool bIsPlayingBackward, int32& InOutBoundsSampleIndex, FBox& OutBounds) override;
 	virtual const FGeometryCacheTrackSampleInfo& GetSampleInfo(float Time, const bool bLooping) override;
 	virtual bool GetMeshDataAtTime(float Time, FGeometryCacheMeshData& OutMeshData) override;
+	virtual bool GetMeshDataAtSampleIndex(int32 SampleIndex, FGeometryCacheMeshData& OutMeshData) override;
 	virtual void UpdateTime(float Time, bool bLooping) override;
 	//~ End UGeometryCacheTrack Interface.
 
diff --git a/Engine/Plugins/Importers/USDImporter/Source/GeometryCacheUSD/Private/GeometryCacheTrackUSD.cpp b/Engine/Plugins/Importers/USDImporter/Source/GeometryCacheUSD/Private/GeometryCacheTrackUSD.cpp
--- a/Engine/Plugins/Importers/USDImporter/Source/GeometryCacheUSD/Private/GeometryCacheTrackUSD.cpp
+++ b/Engine/Plugins/Importers/USDImporter/Source/GeometryCacheUSD/Private/GeometryCacheTrackUSD.cpp
@@ -181,6 +181,11 @@ bool UGeometryCacheTrackUsd::GetMeshDataAtTime(float Time, FGeometryCacheMeshDat
 	return GetMeshData(SampleIndex, OutMeshData);
 }
 
+bool UGeometryCacheTrackUsd::GetMeshDataAtSampleIndex(int32 SampleIndex, FGeometryCacheMeshData& OutMeshData)
+{
+	return GetMeshData(SampleIndex, OutMeshData);
+}
+
 bool UGeometryCacheTrackUsd::GetMeshData(int32 SampleIndex, FGeometryCacheMeshData& OutMeshData)
 {
 	if (IGeometryCacheStreamer::Get().IsTrackRegistered(this))
diff --git a/Engine/Plugins/Importers/USDImporter/Source/GeometryCacheUSD/Public/GeometryCacheTrackUSD.h b/Engine/Plugins/Importers/USDImporter/Source/GeometryCacheUSD/Public/GeometryCacheTrackUSD.h
--- a/Engine/Plugins/Importers/USDImporter/Source/GeometryCacheUSD/Public/GeometryCacheTrackUSD.h
+++ b/Engine/Plugins/Importers/USDImporter/Source/GeometryCacheUSD/Public/GeometryCacheTrackUSD.h
@@ -49,6 +49,7 @@ public:
 	virtual const bool UpdateBoundsData(const float Time, const bool bLooping, const bool bIsPlayingBackward, int32& InOutBoundsSampleIndex, FBox& OutBounds) override;
 	virtual const FGeometryCacheTrackSampleInfo& GetSampleInfo(float Time, const bool bLooping) override;
 	virtual bool GetMeshDataAtTime(float Time, FGeometryCacheMeshData& OutMeshData) override;
+	virtual bool GetMeshDataAtSampleIndex(int32 SampleIndex, FGeometryCacheMeshData& OutMeshData) override;
 	virtual void UpdateTime(float Time, bool bLooping) override;
 	//~ End UGeometryCacheTrack Interface.
 
diff --git a/Engine/Plugins/Runtime/GeometryCache/Source/GeometryCache/Classes/GeometryCacheTrack.h b/Engine/Plugins/Runtime/GeometryCache/Source/GeometryCache/Classes/GeometryCacheTrack.h
--- a/Engine/Plugins/Runtime/GeometryCache/Source/GeometryCache/Classes/GeometryCacheTrack.h
+++ b/Engine/Plugins/Runtime/GeometryCache/Source/GeometryCache/Classes/GeometryCacheTrack.h
@@ -123,6 +123,15 @@ class GEOMETRYCACHE_API UGeometryCacheTrack : public UObject
 	 */
 	virtual bool GetMeshDataAtTime(float Time, FGeometryCacheMeshData& OutMeshData) { return false; }
 
+	/**
+	 * Get the mesh data for the specified sample index
+	 *
+	 * @param SampleIndex - The sample index to get the mesh data for
+	 * @param OutMeshData - Will contain the mesh data that was retrieved, if successful
+	 * @return True if the mesh data could be retrieved
+	 */
+	virtual bool GetMeshDataAtSampleIndex(int32 SampleIndex, FGeometryCacheMeshData& OutMeshData) { return false; }
+
 	/** Return the hash of the mesh data of the track */
 	virtual uint64 GetHash() const { return 0; }
 
diff --git a/Engine/Plugins/Runtime/GeometryCache/Source/GeometryCache/Classes/GeometryCacheTrackStreamable.h b/Engine/Plugins/Runtime/GeometryCache/Source/GeometryCache/Classes/GeometryCacheTrackStreamable.h
--- a/Engine/Plugins/Runtime/GeometryCache/Source/GeometryCache/Classes/GeometryCacheTrackStreamable.h
+++ b/Engine/Plugins/Runtime/GeometryCache/Source/GeometryCache/Classes/GeometryCacheTrackStreamable.h
@@ -113,6 +113,7 @@ class GEOMETRYCACHE_API UGeometryCacheTrackStreamable : public UGeometryCacheTra
 	virtual void SetDuration(float NewDuration) override;
 	virtual const FGeometryCacheTrackSampleInfo& GetSampleInfo(float Time, const bool bLooping) override;
 	virtual bool GetMeshDataAtTime(float Time, FGeometryCacheMeshData& OutMeshData) override;
+	virtual bool GetMeshDataAtSampleIndex(int32 SampleIndex, FGeometryCacheMeshData& OutMeshData) override;
 	virtual uint64 GetHash() const override;
 	//~ End UGeometryCacheTrack Interface.
 
diff --git a/Engine/Plugins/Runtime/GeometryCache/Source/GeometryCache/Private/GeometryCacheTrackStreamable.cpp b/Engine/Plugins/Runtime/GeometryCache/Source/GeometryCache/Private/GeometryCacheTrackStreamable.cpp
--- a/Engine/Plugins/Runtime/GeometryCache/Source/GeometryCache/Private/GeometryCacheTrackStreamable.cpp
+++ b/Engine/Plugins/Runtime/GeometryCache/Source/GeometryCache/Private/GeometryCacheTrackStreamable.cpp
@@ -551,10 +551,19 @@ bool UGeometryCacheTrackStreamable::GetMeshDataAtTime(float Time, FGeometryCache
 		return false;
 	}
 
+	const int32 SampleIndex = FindSampleIndexFromTime(Time, false);
+	return GetMeshDataAtSampleIndex(SampleIndex, OutMeshData);
+}
+
+bool UGeometryCacheTrackStreamable::GetMeshDataAtSampleIndex(int32 SampleIndex, FGeometryCacheMeshData& OutMeshData)
+{
+	if (!Codec || !Samples.IsValidIndex(SampleIndex))
+	{
+		return false;
+	}
+
 	// Fetch the mesh data directly bypassing the GeometryCacheStreamingManager
 	FStreamingGeometryCacheData Data(this);
-
-	const int32 SampleIndex = FindSampleIndexFromTime(Time, false);
 	Data.AddNeededChunk(SampleIndex);
 	Data.UpdateStreamingStatus();
 	Data.BlockTillAllRequestsFinished();
