USD:
* Fixed USD geometry cache streams not being unregistered when the stage is unloaded.
* Fixed geometry cache assets being evicted from the asset cache due to their size varying with the streamer settings.

diff --git a/Engine/Plugins/Importers/USDImporter/Source/GeometryCacheUSD/Private/GeometryCacheTrackUSD.cpp b/Engine/Plugins/Importers/USDImporter/Source/GeometryCacheUSD/Private/GeometryCacheTrackUSD.cpp
--- a/Engine/Plugins/Importers/USDImporter/Source/GeometryCacheUSD/Private/GeometryCacheTrackUSD.cpp
+++ b/Engine/Plugins/Importers/USDImporter/Source/GeometryCacheUSD/Private/GeometryCacheTrackUSD.cpp
@@ -27,7 +27,7 @@ void UGeometryCacheTrackUsd::BeginDestroy()
 {
 	UnloadUsdStage();
 
-	IGeometryCacheStreamer::Get().UnregisterTrack(this);
+	UnregisterStream();
 	UsdStream.Reset();
 
 	Super::BeginDestroy();
@@ -37,26 +37,12 @@ void UGeometryCacheTrackUsd::GetResourceSizeEx(FResourceSizeEx& CumulativeResour
 {
 	Super::GetResourceSizeEx(CumulativeResourceSize);
 
+	// Include only memory usage for data that is part of the track itself
+	// Stream data should only be relevant for the streamer, not the asset cache
+
 	// This is an additional copy that lives on the track
 	MeshData.GetResourceSizeEx(CumulativeResourceSize);
 
-	if (FGeometryCacheUsdStream* UsdStreamPtr = UsdStream.Get())
-	{
-		const FGeometryCacheStreamStats& Stats = UsdStreamPtr->GetStreamStats();
-
-		int32 NumFramesFullyLoaded = Stats.NumCachedFrames;
-		float TotalMemoryForLoadedFramesMB = Stats.MemoryUsed;
-		float MemoryPerFrameMB = TotalMemoryForLoadedFramesMB / NumFramesFullyLoaded;
-
-		// Also account for frames we're going to prefetch
-		uint32 NumRemainingFrames = UsdStreamPtr->GetNumFramesNeeded();
-		float TotalMemoryMB = TotalMemoryForLoadedFramesMB + NumRemainingFrames * MemoryPerFrameMB;
-
-		// Using 1048576 instead of 1000000 for megabyte as that's what's used by FGeometryCacheStreamBase memory stats
-		double TotalMemoryBytes = TotalMemoryMB * 1048576;
-		CumulativeResourceSize.AddDedicatedSystemMemoryBytes(static_cast<SIZE_T>(TotalMemoryBytes + 0.5));
-	}
-
 	CumulativeResourceSize.AddDedicatedSystemMemoryBytes(SampleInfos.GetAllocatedSize());
 }
 
@@ -264,8 +250,6 @@ void UGeometryCacheTrackUsd::Initialize(
 	Duration = NumFrames / FramesPerSecond;
 
 	UsdStream.Reset(new FGeometryCacheUsdStream(this, InReadFunc));
-	IGeometryCacheStreamer::Get().RegisterTrack(this, UsdStream.Get());
-	UsdStream->Prefetch(StartFrameIndex);
 }
 
 void UGeometryCacheTrackUsd::Initialize(
@@ -289,3 +273,18 @@ void UGeometryCacheTrackUsd::UpdateTime(float Time, bool bLooping)
 		UsdStream->UpdateCurrentFrameIndex(FrameIndex);
 	}
 }
+
+void UGeometryCacheTrackUsd::RegisterStream()
+{
+	const bool bNeedPrefetch = !IGeometryCacheStreamer::Get().IsTrackRegistered(this);
+	IGeometryCacheStreamer::Get().RegisterTrack(this, UsdStream.Get());
+	if (bNeedPrefetch)
+	{
+		UsdStream->Prefetch(StartFrameIndex);
+	}
+}
+
+void UGeometryCacheTrackUsd::UnregisterStream()
+{
+	IGeometryCacheStreamer::Get().UnregisterTrack(this);
+}
diff --git a/Engine/Plugins/Importers/USDImporter/Source/GeometryCacheUSD/Private/GeometryCacheUSDComponent.cpp b/Engine/Plugins/Importers/USDImporter/Source/GeometryCacheUSD/Private/GeometryCacheUSDComponent.cpp
--- a/Engine/Plugins/Importers/USDImporter/Source/GeometryCacheUSD/Private/GeometryCacheUSDComponent.cpp
+++ b/Engine/Plugins/Importers/USDImporter/Source/GeometryCacheUSD/Private/GeometryCacheUSDComponent.cpp
@@ -1,6 +1,9 @@
 // Copyright Epic Games, Inc. All Rights Reserved.
 
 #include "GeometryCacheUSDComponent.h"
+
+#include "GeometryCache.h"
+#include "GeometryCacheTrackUSD.h"
 #include "GeometryCacheUSDSceneProxy.h"
 
 FPrimitiveSceneProxy* UGeometryCacheUsdComponent::CreateSceneProxy()
@@ -17,3 +20,42 @@ void UGeometryCacheUsdComponent::PostDuplicate(bool bDuplicateForPIE)
 
 	SetupTrackData();
 }
+
+void UGeometryCacheUsdComponent::OnRegister()
+{
+	if (GeometryCache != nullptr)
+	{
+		for (UGeometryCacheTrack* Track : GeometryCache->Tracks)
+		{
+			if (UGeometryCacheTrackUsd* UsdTrack = Cast<UGeometryCacheTrackUsd>(Track))
+			{
+				UsdTrack->RegisterStream();
+			}
+		}
+	}
+
+	ClearTrackData();
+	SetupTrackData();
+
+	// Skip code for UGeometryCacheComponent::OnRegister
+	UMeshComponent::OnRegister();
+}
+
+void UGeometryCacheUsdComponent::OnUnregister()
+{
+	if (GeometryCache != nullptr)
+	{
+		for (UGeometryCacheTrack* Track : GeometryCache->Tracks)
+		{
+			if (UGeometryCacheTrackUsd* UsdTrack = Cast<UGeometryCacheTrackUsd>(Track))
+			{
+				UsdTrack->UnregisterStream();
+			}
+		}
+	}
+
+	ClearTrackData();
+
+	// Skip code for UGeometryCacheComponent::OnUnregister
+	UMeshComponent::OnUnregister();
+}
diff --git a/Engine/Plugins/Importers/USDImporter/Source/GeometryCacheUSD/Public/GeometryCacheTrackUSD.h b/Engine/Plugins/Importers/USDImporter/Source/GeometryCacheUSD/Public/GeometryCacheTrackUSD.h
--- a/Engine/Plugins/Importers/USDImporter/Source/GeometryCacheUSD/Public/GeometryCacheTrackUSD.h
+++ b/Engine/Plugins/Importers/USDImporter/Source/GeometryCacheUSD/Public/GeometryCacheTrackUSD.h
@@ -65,6 +65,9 @@ public:
 	// Discards our CurrentStagePinned to release the stage
 	void UnloadUsdStage();
 
+	void RegisterStream();
+	void UnregisterStream();
+
 public:
 	double FramesPerSecond;
 	int32 StartFrameIndex;
diff --git a/Engine/Plugins/Importers/USDImporter/Source/GeometryCacheUSD/Public/GeometryCacheUSDComponent.h b/Engine/Plugins/Importers/USDImporter/Source/GeometryCacheUSD/Public/GeometryCacheUSDComponent.h
--- a/Engine/Plugins/Importers/USDImporter/Source/GeometryCacheUSD/Public/GeometryCacheUSDComponent.h
+++ b/Engine/Plugins/Importers/USDImporter/Source/GeometryCacheUSD/Public/GeometryCacheUSDComponent.h
@@ -21,4 +21,9 @@ public:
 	//~ Begin UPrimitiveComponent Interface.
 	virtual FPrimitiveSceneProxy* CreateSceneProxy() override;
 	//~ End UPrimitiveComponent Interface.
+
+	//~ Begin UActorComponent Interface.
+	virtual void OnRegister() override;
+	virtual void OnUnregister() override;
+	//~ End UActorComponent Interface.
 };
diff --git a/Engine/Plugins/Runtime/GeometryCache/Source/GeometryCacheStreamer/Private/GeometryCacheStreamer.cpp b/Engine/Plugins/Runtime/GeometryCache/Source/GeometryCacheStreamer/Private/GeometryCacheStreamer.cpp
--- a/Engine/Plugins/Runtime/GeometryCache/Source/GeometryCacheStreamer/Private/GeometryCacheStreamer.cpp
+++ b/Engine/Plugins/Runtime/GeometryCache/Source/GeometryCacheStreamer/Private/GeometryCacheStreamer.cpp
@@ -6,6 +6,7 @@
 #include "GeometryCacheModule.h"
 #include "GeometryCacheStreamerSettings.h"
 #include "IGeometryCacheStream.h"
+#include "Templates/RefCounting.h"
 #include "Widgets/Notifications/SNotificationList.h"
 
 #define LOCTEXT_NAMESPACE "GeometryCacheStreamer"
@@ -38,10 +39,10 @@ public:
 	//~ Begin IGeometryCacheStreamer Interface
 	static IGeometryCacheStreamer& Get();
 
-	virtual void RegisterTrack(UGeometryCacheTrack* AbcTrack, IGeometryCacheStream* Stream) override;
-	virtual void UnregisterTrack(UGeometryCacheTrack* AbcTrack) override;
-	virtual bool IsTrackRegistered(UGeometryCacheTrack* AbcTrack) const override;
-	virtual bool TryGetFrameData(UGeometryCacheTrack* AbcTrack, int32 FrameIndex, FGeometryCacheMeshData& OutMeshData) override;
+	virtual void RegisterTrack(UGeometryCacheTrack* Track, IGeometryCacheStream* Stream) override;
+	virtual void UnregisterTrack(UGeometryCacheTrack* Track) override;
+	virtual bool IsTrackRegistered(UGeometryCacheTrack* Track) const override;
+	virtual bool TryGetFrameData(UGeometryCacheTrack* Track, int32 FrameIndex, FGeometryCacheMeshData& OutMeshData) override;
 	//~ End IGeometryCacheStreamer Interface
 
 private:
@@ -50,7 +51,19 @@ private:
 private:
 	FTSTicker::FDelegateHandle TickHandle;
 
-	typedef TMap<UGeometryCacheTrack*, IGeometryCacheStream*> FTracksToStreams;
+	class FRefCountedStream : public FThreadSafeRefCountedObject
+	{
+	public:
+		FRefCountedStream(IGeometryCacheStream* InStream)
+		: Stream(InStream)
+		{
+			AddRef();
+		}
+
+		IGeometryCacheStream* Stream = nullptr;
+	};
+
+	typedef TMap<UGeometryCacheTrack*, FRefCountedStream*> FTracksToStreams;
 	FTracksToStreams TracksToStreams;
 
 	const int32 MaxReads;
@@ -100,7 +113,7 @@ void FGeometryCacheStreamer::Tick(float Time)
 		for (const auto& Pair : TracksToStreams)
 		{
 			TArray<int32> FramesCompleted;
-			IGeometryCacheStream* Stream = Pair.Value;
+			IGeometryCacheStream* Stream = Pair.Value->Stream;
 			Stream->UpdateRequestStatus(FramesCompleted);
 			NumReads -= FramesCompleted.Num();
 			NumFramesToStream += Stream->GetNumFramesNeeded();
@@ -111,7 +124,7 @@ void FGeometryCacheStreamer::Tick(float Time)
 		int32 AvailableReads = MaxReads - NumReads;
 		if (NumStreams > 0 && AvailableReads > 0)
 		{
-			TArray<IGeometryCacheStream*> Streams;
+			TArray<FRefCountedStream*> Streams;
 			TracksToStreams.GenerateValueArray(Streams);
 
 			// Streams are checked round robin until there are no more reads available
@@ -131,7 +144,7 @@ void FGeometryCacheStreamer::Tick(float Time)
 					continue;
 				}
 
-				IGeometryCacheStream* Stream = Streams[CurrentIndex];
+				IGeometryCacheStream* Stream = Streams[CurrentIndex]->Stream;
 				if (Stream->GetNumFramesNeeded() > 0)
 				{
 					if (Stream->RequestFrameData())
@@ -206,7 +219,7 @@ void FGeometryCacheStreamer::BalanceStreams()
 	FGeometryCacheStreamStats AggregatedStats;
 	for (const auto& Pair : TracksToStreams)
 	{
-		IGeometryCacheStream* Stream = Pair.Value;
+		IGeometryCacheStream* Stream = Pair.Value->Stream;
 		const FGeometryCacheStreamStats& StreamStats = Stream->GetStreamStats();
 		AggregatedStats.MemoryUsed += StreamStats.MemoryUsed;
 		AggregatedStats.CachedDuration = FMath::Max(AggregatedStats.CachedDuration, StreamStats.CachedDuration);
@@ -231,7 +244,7 @@ void FGeometryCacheStreamer::BalanceStreams()
 		const float AvgDuration = MemoryLimit / AggregatedStats.AverageBitrate;
 		for (const auto& Pair : TracksToStreams)
 		{
-			IGeometryCacheStream* Stream = Pair.Value;
+			IGeometryCacheStream* Stream = Pair.Value->Stream;
 			const FGeometryCacheStreamStats& StreamStats = Stream->GetStreamStats();
 
 			const float MaxDuration = FMath::Clamp(AvgDuration, 0.f, DurationLimit);
@@ -241,34 +254,43 @@ void FGeometryCacheStreamer::BalanceStreams()
 	}
 }
 
-void FGeometryCacheStreamer::RegisterTrack(UGeometryCacheTrack* AbcTrack, IGeometryCacheStream* Stream)
+void FGeometryCacheStreamer::RegisterTrack(UGeometryCacheTrack* Track, IGeometryCacheStream* Stream)
 {
-	check(AbcTrack && Stream && !TracksToStreams.Contains(AbcTrack));
-
-	TracksToStreams.Add(AbcTrack, Stream);
+	if (FRefCountedStream** RefCountStream = TracksToStreams.Find(Track))
+	{
+		(*RefCountStream)->AddRef();
+	}
+	else
+	{
+		TracksToStreams.Add(Track, new FRefCountedStream(Stream));
+	}
 }
 
-void FGeometryCacheStreamer::UnregisterTrack(UGeometryCacheTrack* AbcTrack)
+void FGeometryCacheStreamer::UnregisterTrack(UGeometryCacheTrack* Track)
 {
-	if (IGeometryCacheStream** Stream = TracksToStreams.Find(AbcTrack))
+	if (FRefCountedStream** RefCountStream = TracksToStreams.Find(Track))
 	{
-		int32 NumRequests = (*Stream)->CancelRequests();
+		int32 NumRequests = ((*RefCountStream)->Stream)->CancelRequests();
 		NumReads -= NumRequests;
 
-		TracksToStreams.Remove(AbcTrack);
+		// The refcounted object deletes itself when the count reaches 0
+		if ((*RefCountStream)->Release() == 0)
+		{
+			TracksToStreams.Remove(Track);
+		}
 	}
 }
 
-bool FGeometryCacheStreamer::IsTrackRegistered(UGeometryCacheTrack* AbcTrack) const
+bool FGeometryCacheStreamer::IsTrackRegistered(UGeometryCacheTrack* Track) const
 {
-	return TracksToStreams.Contains(AbcTrack);
+	return TracksToStreams.Contains(Track);
 }
 
-bool FGeometryCacheStreamer::TryGetFrameData(UGeometryCacheTrack* AbcTrack, int32 FrameIndex, FGeometryCacheMeshData& OutMeshData)
+bool FGeometryCacheStreamer::TryGetFrameData(UGeometryCacheTrack* Track, int32 FrameIndex, FGeometryCacheMeshData& OutMeshData)
 {
-	if (IGeometryCacheStream** Stream = TracksToStreams.Find(AbcTrack))
+	if (FRefCountedStream** RefCountStream = TracksToStreams.Find(Track))
 	{
-		if ((*Stream)->GetFrameData(FrameIndex, OutMeshData))
+		if (((*RefCountStream)->Stream)->GetFrameData(FrameIndex, OutMeshData))
 		{
 			return true;
 		}
