- Add ConditionalPostLoad for Niagara material texture bindings / MIC materials
- Without this duplicated a system for validation can crash as we access a subobject which has not been PostLoaded yet

diff --git a/Engine/Plugins/FX/Niagara/Source/Niagara/Private/NiagaraDecalRendererProperties.cpp b/Engine/Plugins/FX/Niagara/Source/Niagara/Private/NiagaraDecalRendererProperties.cpp
--- a/Engine/Plugins/FX/Niagara/Source/Niagara/Private/NiagaraDecalRendererProperties.cpp
+++ b/Engine/Plugins/FX/Niagara/Source/Niagara/Private/NiagaraDecalRendererProperties.cpp
@@ -102,6 +102,7 @@ void UNiagaraDecalRendererProperties::PostLoad()
 	ChangeToPositionBinding(PositionBinding);
 #endif
 	PostLoadBindings(ENiagaraRendererSourceDataMode::Particles);
+	MaterialParameters.ConditionalPostLoad();
 }
 
 void UNiagaraDecalRendererProperties::PostInitProperties()
diff --git a/Engine/Plugins/FX/Niagara/Source/Niagara/Private/NiagaraMeshRendererProperties.cpp b/Engine/Plugins/FX/Niagara/Source/Niagara/Private/NiagaraMeshRendererProperties.cpp
--- a/Engine/Plugins/FX/Niagara/Source/Niagara/Private/NiagaraMeshRendererProperties.cpp
+++ b/Engine/Plugins/FX/Niagara/Source/Niagara/Private/NiagaraMeshRendererProperties.cpp
@@ -948,6 +948,18 @@ void UNiagaraMeshRendererProperties::PostLoad()
 		}
 	}
 
+	for (const FNiagaraMeshMICOverride& MICOverrideMaterial : MICOverrideMaterials)
+	{
+		if (MICOverrideMaterial.OriginalMaterial)
+		{
+			MICOverrideMaterial.OriginalMaterial->ConditionalPostLoad();
+		}
+		if (MICOverrideMaterial.ReplacementMaterial)
+		{
+			MICOverrideMaterial.ReplacementMaterial->ConditionalPostLoad();
+		}
+	}
+
 #if WITH_EDITORONLY_DATA
 	if (MaterialParameterBindings_DEPRECATED.Num() > 0)
 	{
@@ -955,6 +967,7 @@ void UNiagaraMeshRendererProperties::PostLoad()
 		MaterialParameterBindings_DEPRECATED.Empty();
 	}
 #endif
+	MaterialParameters.ConditionalPostLoad();
 }
 
 #if WITH_EDITORONLY_DATA
diff --git a/Engine/Plugins/FX/Niagara/Source/Niagara/Private/NiagaraRendererProperties.cpp b/Engine/Plugins/FX/Niagara/Source/Niagara/Private/NiagaraRendererProperties.cpp
--- a/Engine/Plugins/FX/Niagara/Source/Niagara/Private/NiagaraRendererProperties.cpp
+++ b/Engine/Plugins/FX/Niagara/Source/Niagara/Private/NiagaraRendererProperties.cpp
@@ -180,6 +180,17 @@ void FNiagaraRendererLayout::Finalize()
 
 //////////////////////////////////////////////////////////////////////////
 
+void FNiagaraRendererMaterialParameters::ConditionalPostLoad()
+{
+	for (const FNiagaraRendererMaterialTextureParameter& TextureParameter : TextureParameters)
+	{
+		if (TextureParameter.Texture)
+		{
+			TextureParameter.Texture->ConditionalPostLoad();
+		}
+	}
+}
+
 #if WITH_EDITORONLY_DATA
 void FNiagaraRendererMaterialParameters::RenameVariable(const FNiagaraVariableBase& OldVariable, const FNiagaraVariableBase& NewVariable, const FVersionedNiagaraEmitter& InEmitter, ENiagaraRendererSourceDataMode SourceMode)
 {
diff --git a/Engine/Plugins/FX/Niagara/Source/Niagara/Private/NiagaraRibbonRendererProperties.cpp b/Engine/Plugins/FX/Niagara/Source/Niagara/Private/NiagaraRibbonRendererProperties.cpp
--- a/Engine/Plugins/FX/Niagara/Source/Niagara/Private/NiagaraRibbonRendererProperties.cpp
+++ b/Engine/Plugins/FX/Niagara/Source/Niagara/Private/NiagaraRibbonRendererProperties.cpp
@@ -182,6 +182,7 @@ void UNiagaraRibbonRendererProperties::PostLoad()
 		MaterialParameterBindings_DEPRECATED.Empty();
 	}
 #endif
+	MaterialParameters.ConditionalPostLoad();
 }
 
 FNiagaraBoundsCalculator* UNiagaraRibbonRendererProperties::CreateBoundsCalculator()
diff --git a/Engine/Plugins/FX/Niagara/Source/Niagara/Private/NiagaraSpriteRendererProperties.cpp b/Engine/Plugins/FX/Niagara/Source/Niagara/Private/NiagaraSpriteRendererProperties.cpp
--- a/Engine/Plugins/FX/Niagara/Source/Niagara/Private/NiagaraSpriteRendererProperties.cpp
+++ b/Engine/Plugins/FX/Niagara/Source/Niagara/Private/NiagaraSpriteRendererProperties.cpp
@@ -187,6 +187,7 @@ void UNiagaraSpriteRendererProperties::PostLoad()
 		MaterialParameterBindings_DEPRECATED.Empty();
 	}
 #endif
+	MaterialParameters.ConditionalPostLoad();
 }
 
 void UNiagaraSpriteRendererProperties::PostInitProperties()
diff --git a/Engine/Plugins/FX/Niagara/Source/Niagara/Private/NiagaraVolumeRendererProperties.cpp b/Engine/Plugins/FX/Niagara/Source/Niagara/Private/NiagaraVolumeRendererProperties.cpp
--- a/Engine/Plugins/FX/Niagara/Source/Niagara/Private/NiagaraVolumeRendererProperties.cpp
+++ b/Engine/Plugins/FX/Niagara/Source/Niagara/Private/NiagaraVolumeRendererProperties.cpp
@@ -98,6 +98,7 @@ void UNiagaraVolumeRendererProperties::PostLoad()
 	ChangeToPositionBinding(PositionBinding);
 #endif
 	PostLoadBindings(ENiagaraRendererSourceDataMode::Particles);
+	MaterialParameters.ConditionalPostLoad();
 }
 
 void UNiagaraVolumeRendererProperties::PostInitProperties()
diff --git a/Engine/Plugins/FX/Niagara/Source/Niagara/Public/NiagaraRendererProperties.h b/Engine/Plugins/FX/Niagara/Source/Niagara/Public/NiagaraRendererProperties.h
--- a/Engine/Plugins/FX/Niagara/Source/Niagara/Public/NiagaraRendererProperties.h
+++ b/Engine/Plugins/FX/Niagara/Source/Niagara/Public/NiagaraRendererProperties.h
@@ -258,6 +258,7 @@ struct FNiagaraRendererMaterialParameters
 	UPROPERTY(EditAnywhere, Category = "Material")
 	TArray<FNiagaraRendererMaterialStaticBoolParameter> StaticBoolParameters;
 
+	NIAGARA_API void ConditionalPostLoad();
 #if WITH_EDITORONLY_DATA
 	NIAGARA_API void RenameVariable(const FNiagaraVariableBase& OldVariable, const FNiagaraVariableBase& NewVariable, const FVersionedNiagaraEmitter& InEmitter, ENiagaraRendererSourceDataMode SourceMode);
 	NIAGARA_API void RemoveVariable(const FNiagaraVariableBase& OldVariable, const FVersionedNiagaraEmitter& InEmitter, ENiagaraRendererSourceDataMode SourceMode);
diff --git a/Engine/Plugins/Runtime/GeometryCache/Source/GeometryCache/Private/NiagaraGeometryCacheRendererProperties.cpp b/Engine/Plugins/Runtime/GeometryCache/Source/GeometryCache/Private/NiagaraGeometryCacheRendererProperties.cpp
--- a/Engine/Plugins/Runtime/GeometryCache/Source/GeometryCache/Private/NiagaraGeometryCacheRendererProperties.cpp
+++ b/Engine/Plugins/Runtime/GeometryCache/Source/GeometryCache/Private/NiagaraGeometryCacheRendererProperties.cpp
@@ -78,6 +78,29 @@ void UNiagaraGeometryCacheRendererProperties::PostLoad()
 {
 	Super::PostLoad();
 	PostLoadBindings(ENiagaraRendererSourceDataMode::Particles);
+	MaterialParameters.ConditionalPostLoad();
+	for (const FNiagaraGeometryCacheReference& GeoCache : GeometryCaches)
+	{
+		for (UMaterialInterface* Material : GeoCache.OverrideMaterials)
+		{
+			if (Material)
+			{
+				Material->ConditionalPostLoad();
+			}
+		}
+		for (const FNiagaraGeometryCacheMICOverride& MICOverride : GeoCache.MICOverrideMaterials)
+		{
+			if (MICOverride.OriginalMaterial)
+			{
+				MICOverride.OriginalMaterial->ConditionalPostLoad();
+			}
+
+			if (MICOverride.ReplacementMaterial)
+			{
+				MICOverride.ReplacementMaterial->ConditionalPostLoad();
+			}
+		}
+	}
 }
 
 template<typename T>
