Fixed parallel GDME issues and re-enabled.
 - Hair visualization mode context needs a mutex to initialize. Hair also currently requires deterministic ordering of dynamic meshes to avoid artifacts, so disabling parallel support for now.
 - Fixed race in allocating particle order buffer.
 - Fixed Geometry Collection UpdateFrame to support a command list.

diff --git a/Engine/Plugins/Runtime/GeometryCache/Source/GeometryCache/Private/GeometryCacheComponent.cpp b/Engine/Plugins/Runtime/GeometryCache/Source/GeometryCache/Private/GeometryCacheComponent.cpp
--- a/Engine/Plugins/Runtime/GeometryCache/Source/GeometryCache/Private/GeometryCacheComponent.cpp
+++ b/Engine/Plugins/Runtime/GeometryCache/Source/GeometryCache/Private/GeometryCacheComponent.cpp
@@ -139,7 +139,7 @@ void UGeometryCacheComponent::TickComponent(float DeltaTime, enum ELevelTick Tic
 			ENQUEUE_RENDER_COMMAND(FGeometryCacheUpdateAnimation)(
 				[InSceneProxy, AnimationTime, bInLooping, bIsPlayingBackwards, InPlaybackSpeed, ActualPlaybackSpeed, InMotionVectorScale](FRHICommandList& RHICmdList)
 				{
-					InSceneProxy->UpdateAnimation(AnimationTime, bInLooping, bIsPlayingBackwards, InPlaybackSpeed, InMotionVectorScale);
+					InSceneProxy->UpdateAnimation(RHICmdList, AnimationTime, bInLooping, bIsPlayingBackwards, InPlaybackSpeed, InMotionVectorScale);
 				});
 		}
 	}
@@ -196,7 +196,7 @@ void UGeometryCacheComponent::TickAtThisTime(const float Time, bool bInIsRunning
 			ENQUEUE_RENDER_COMMAND(FGeometryCacheUpdateAnimation)(
 				[InSceneProxy, AnimationTime, bInIsLooping, bInBackwards, InPlaybackSpeed, InMotionVectorScale](FRHICommandList& RHICmdList)
 				{
-					InSceneProxy->UpdateAnimation(AnimationTime, bInIsLooping, bInBackwards, InPlaybackSpeed, InMotionVectorScale);
+					InSceneProxy->UpdateAnimation(RHICmdList, AnimationTime, bInIsLooping, bInBackwards, InPlaybackSpeed, InMotionVectorScale);
 				});
 		}
 	}
diff --git a/Engine/Plugins/Runtime/GeometryCache/Source/GeometryCache/Private/GeometryCacheSceneProxy.cpp b/Engine/Plugins/Runtime/GeometryCache/Source/GeometryCache/Private/GeometryCacheSceneProxy.cpp
--- a/Engine/Plugins/Runtime/GeometryCache/Source/GeometryCache/Private/GeometryCacheSceneProxy.cpp
+++ b/Engine/Plugins/Runtime/GeometryCache/Source/GeometryCache/Private/GeometryCacheSceneProxy.cpp
@@ -109,7 +109,10 @@ FGeometryCacheSceneProxy::FGeometryCacheSceneProxy(UGeometryCacheComponent* Comp
 
 		if (SampleInfo.NumVertices > 0)
 		{
-			NewSection->InitRenderResources(SampleInfo.NumVertices, SampleInfo.NumIndices);
+			ENQUEUE_RENDER_COMMAND(FGeometryCacheInitResources)([NewSection, NumVertices = SampleInfo.NumVertices, NumIndices = SampleInfo.NumIndices] (FRHICommandListBase& RHICmdList)
+			{
+				NewSection->InitRenderResources(RHICmdList, NumVertices, NumIndices);
+			});
 		}
 
 		// Grab materials
@@ -154,7 +157,7 @@ FGeometryCacheSceneProxy::FGeometryCacheSceneProxy(UGeometryCacheComponent* Comp
 		ENQUEUE_RENDER_COMMAND(FGeometryCacheUpdateAnimation)(
 			[SceneProxy](FRHICommandListImmediate& RHICmdList)
 		{
-			SceneProxy->FrameUpdate();
+			SceneProxy->FrameUpdate(RHICmdList);
 		});
 
 #if RHI_RAYTRACING
@@ -446,7 +449,7 @@ void FGeometryCacheSceneProxy::GetDynamicMeshElements(const TArray<const FSceneV
 		{
 			// When ray tracing is disabled, update only when visible
 			// This is the old behavior
-			FrameUpdate();
+			FrameUpdate(RHICmdList);
 		}
 
 		// Iterate over all batches in all tracks and add them to all the relevant views	
@@ -588,7 +591,7 @@ uint32 FGeometryCacheSceneProxy::GetAllocatedSize(void) const
 	return(FPrimitiveSceneProxy::GetAllocatedSize());
 }
 
-void FGeometryCacheSceneProxy::UpdateAnimation(float NewTime, bool bNewLooping, bool bNewIsPlayingBackwards, float NewPlaybackSpeed, float NewMotionVectorScale)
+void FGeometryCacheSceneProxy::UpdateAnimation(FRHICommandListBase& RHICmdList, float NewTime, bool bNewLooping, bool bNewIsPlayingBackwards, float NewPlaybackSpeed, float NewMotionVectorScale)
 {
 	Time = NewTime;
 	bLooping = bNewLooping;
@@ -600,7 +603,7 @@ void FGeometryCacheSceneProxy::UpdateAnimation(float NewTime, bool bNewLooping,
 	if (IsRayTracingEnabled())
 	{
 		// When ray tracing is enabled, update regardless of visibility
-		FrameUpdate();
+		FrameUpdate(RHICmdList);
 
 #if RHI_RAYTRACING
 		for (FGeomCacheTrackProxy* Section : Tracks)
@@ -666,8 +669,6 @@ void FGeometryCacheSceneProxy::UpdateAnimation(float NewTime, bool bNewLooping,
 
 				if (Segments.Num() > 0)
 				{
-					FRHICommandList& RHICmdList = FRHICommandListImmediate::Get();
-
 					if (bRequireRecreate)
 					{
 						Section->RayTracingGeometry.UpdateRHI(RHICmdList);
@@ -679,7 +680,7 @@ void FGeometryCacheSceneProxy::UpdateAnimation(float NewTime, bool bNewLooping,
 						BuildParams.Geometry = Section->RayTracingGeometry.RayTracingGeometryRHI;
 						BuildParams.BuildMode = EAccelerationStructureBuildMode::Build;
 						BuildParams.Segments = Section->RayTracingGeometry.Initializer.Segments;
-						RHICmdList.BuildAccelerationStructures(MakeArrayView(&BuildParams, 1));
+						FRHIComputeCommandList::Get(RHICmdList).BuildAccelerationStructures(MakeArrayView(&BuildParams, 1));
 					}
 				}
 			}
@@ -688,7 +689,7 @@ void FGeometryCacheSceneProxy::UpdateAnimation(float NewTime, bool bNewLooping,
 	}
 }
 
-void FGeometryCacheSceneProxy::FrameUpdate() const
+void FGeometryCacheSceneProxy::FrameUpdate(FRHICommandListBase& RHICmdList) const
 {
 	TRACE_CPUPROFILER_EVENT_SCOPE(FGeometryCacheSceneProxy::FrameUpdate);
 
@@ -818,7 +819,7 @@ void FGeometryCacheSceneProxy::FrameUpdate() const
 				}
 				else if (!TrackProxy->bResourcesInitialized)
 				{
-					TrackProxy->InitRenderResources(NumVerts, TrackProxy->MeshData->Indices.Num());
+					TrackProxy->InitRenderResources(RHICmdList, NumVerts, TrackProxy->MeshData->Indices.Num());
 				}
 
 				Scratch.Prepare(NumVerts, bHasMotionVectors);
@@ -1064,18 +1065,18 @@ void FGeometryCacheSceneProxy::FrameUpdate() const
 
 				// Upload other non-motionblurred data
 				if (!TrackProxy->MeshData->VertexInfo.bConstantIndices)
-					TrackProxy->IndexBuffer.Update(TrackProxy->MeshData->Indices);
+					TrackProxy->IndexBuffer.Update(RHICmdList, TrackProxy->MeshData->Indices);
 
 				if (TrackProxy->MeshData->VertexInfo.bHasTangentX)
-					TrackProxy->TangentXBuffer.Update(Scratch.InterpolatedTangentX);
+					TrackProxy->TangentXBuffer.Update(RHICmdList, Scratch.InterpolatedTangentX);
 				if (TrackProxy->MeshData->VertexInfo.bHasTangentZ)
-					TrackProxy->TangentZBuffer.Update(Scratch.InterpolatedTangentZ);
+					TrackProxy->TangentZBuffer.Update(RHICmdList, Scratch.InterpolatedTangentZ);
 
 				if (TrackProxy->MeshData->VertexInfo.bHasUV0)
-					TrackProxy->TextureCoordinatesBuffer.Update(Scratch.InterpolatedUVs);
+					TrackProxy->TextureCoordinatesBuffer.Update(RHICmdList, Scratch.InterpolatedUVs);
 
 				if (TrackProxy->MeshData->VertexInfo.bHasColor0)
-					TrackProxy->ColorBuffer.Update(Scratch.InterpolatedColors);
+					TrackProxy->ColorBuffer.Update(RHICmdList, Scratch.InterpolatedColors);
 
 				bool bIsCompatibleWithCachedFrame = TrackProxy->IsTopologyCompatible(
 					TrackProxy->PositionBufferFrameIndices[TrackProxy->CurrentPositionBufferIndex % 2],
@@ -1086,8 +1087,8 @@ void FGeometryCacheSceneProxy::FrameUpdate() const
 					// Initialize both buffers the first frame
 					if (TrackProxy->CurrentPositionBufferIndex == -1 || !bIsCompatibleWithCachedFrame)
 					{
-						TrackProxy->PositionBuffers[0].Update(Scratch.InterpolatedPositions);
-						TrackProxy->PositionBuffers[1].Update(Scratch.InterpolatedPositions);
+						TrackProxy->PositionBuffers[0].Update(RHICmdList, Scratch.InterpolatedPositions);
+						TrackProxy->PositionBuffers[1].Update(RHICmdList, Scratch.InterpolatedPositions);
 						TrackProxy->CurrentPositionBufferIndex = 0;
 						TrackProxy->PositionBufferFrameTimes[0] = Time;
 						TrackProxy->PositionBufferFrameTimes[1] = Time;
@@ -1101,7 +1102,7 @@ void FGeometryCacheSceneProxy::FrameUpdate() const
 					else
 					{
 						TrackProxy->CurrentPositionBufferIndex++;
-						TrackProxy->PositionBuffers[TrackProxy->CurrentPositionBufferIndex % 2].Update(Scratch.InterpolatedPositions);
+						TrackProxy->PositionBuffers[TrackProxy->CurrentPositionBufferIndex % 2].Update(RHICmdList, Scratch.InterpolatedPositions);
 						TrackProxy->PositionBufferFrameTimes[TrackProxy->CurrentPositionBufferIndex % 2] = Time;
 						TrackProxy->PositionBufferFrameIndices[TrackProxy->CurrentPositionBufferIndex % 2] = TrackProxy->NextFrameIndex;
 					}
@@ -1109,8 +1110,8 @@ void FGeometryCacheSceneProxy::FrameUpdate() const
 				else
 				{
 					TrackProxy->CurrentPositionBufferIndex = 0;
-					TrackProxy->PositionBuffers[0].Update(Scratch.InterpolatedPositions);
-					TrackProxy->PositionBuffers[1].Update(Scratch.InterpolatedMotionVectors);
+					TrackProxy->PositionBuffers[0].Update(RHICmdList, Scratch.InterpolatedPositions);
+					TrackProxy->PositionBuffers[1].Update(RHICmdList, Scratch.InterpolatedMotionVectors);
 					TrackProxy->PositionBufferFrameIndices[0] = TrackProxy->FrameIndex;
 					TrackProxy->PositionBufferFrameIndices[1] = -1;
 					TrackProxy->PositionBufferFrameTimes[0] = Time;
@@ -1135,7 +1136,7 @@ void FGeometryCacheSceneProxy::FrameUpdate() const
 					}
 					else if (!TrackProxy->bResourcesInitialized)
 					{
-						TrackProxy->InitRenderResources(MeshDataToUse->Positions.Num(), MeshDataToUse->Indices.Num());
+						TrackProxy->InitRenderResources(RHICmdList, MeshDataToUse->Positions.Num(), MeshDataToUse->Indices.Num());
 					}
 
 					TrackProxy->bNextFrameMeshDataSelected = bNextFrame;
@@ -1143,18 +1144,18 @@ void FGeometryCacheSceneProxy::FrameUpdate() const
 					const int32 NumVertices = MeshDataToUse->Positions.Num();
 
 					if (MeshDataToUse->VertexInfo.bHasTangentX)
-						TrackProxy->TangentXBuffer.Update(MeshDataToUse->TangentsX);
+						TrackProxy->TangentXBuffer.Update(RHICmdList, MeshDataToUse->TangentsX);
 					if (MeshDataToUse->VertexInfo.bHasTangentZ)
-						TrackProxy->TangentZBuffer.Update(MeshDataToUse->TangentsZ);
+						TrackProxy->TangentZBuffer.Update(RHICmdList, MeshDataToUse->TangentsZ);
 
 					if (!MeshDataToUse->VertexInfo.bConstantIndices)
-						TrackProxy->IndexBuffer.Update(MeshDataToUse->Indices);
+						TrackProxy->IndexBuffer.Update(RHICmdList, MeshDataToUse->Indices);
 
 					if (MeshDataToUse->VertexInfo.bHasUV0)
-						TrackProxy->TextureCoordinatesBuffer.Update(MeshDataToUse->TextureCoordinates);
+						TrackProxy->TextureCoordinatesBuffer.Update(RHICmdList, MeshDataToUse->TextureCoordinates);
 
 					if (MeshDataToUse->VertexInfo.bHasColor0)
-						TrackProxy->ColorBuffer.Update(MeshDataToUse->Colors);
+						TrackProxy->ColorBuffer.Update(RHICmdList, MeshDataToUse->Colors);
 
 					const bool bIsCompatibleWithCachedFrame = TrackProxy->IsTopologyCompatible(
 						TrackProxy->PositionBufferFrameIndices[TrackProxy->CurrentPositionBufferIndex % 2],
@@ -1166,8 +1167,8 @@ void FGeometryCacheSceneProxy::FrameUpdate() const
 						// with a previous buffer referencing a buffer from another topology
 						if (TrackProxy->CurrentPositionBufferIndex == -1 || !bIsCompatibleWithCachedFrame || bSeeked)
 						{
-							TrackProxy->PositionBuffers[0].Update(MeshDataToUse->Positions);
-							TrackProxy->PositionBuffers[1].Update(MeshDataToUse->Positions);
+							TrackProxy->PositionBuffers[0].Update(RHICmdList, MeshDataToUse->Positions);
+							TrackProxy->PositionBuffers[1].Update(RHICmdList, MeshDataToUse->Positions);
 							TrackProxy->CurrentPositionBufferIndex = 0;
 							TrackProxy->PositionBufferFrameIndices[0] = FrameIndexToUse;
 							TrackProxy->PositionBufferFrameIndices[1] = FrameIndexToUse;
@@ -1179,7 +1180,7 @@ void FGeometryCacheSceneProxy::FrameUpdate() const
 						else
 						{
 							TrackProxy->CurrentPositionBufferIndex++;
-							TrackProxy->PositionBuffers[TrackProxy->CurrentPositionBufferIndex % 2].Update(MeshDataToUse->Positions);
+							TrackProxy->PositionBuffers[TrackProxy->CurrentPositionBufferIndex % 2].Update(RHICmdList, MeshDataToUse->Positions);
 							TrackProxy->PositionBufferFrameIndices[TrackProxy->CurrentPositionBufferIndex % 2] = FrameIndexToUse;
 						}
 					}
@@ -1217,8 +1218,8 @@ void FGeometryCacheSceneProxy::FrameUpdate() const
 						TrackProxy->SubframeInterpolationFactor = FMath::IsNearlyEqual(DeltaInterpolationFactor, 1.0f, KINDA_SMALL_NUMBER) ? 1.0f : DeltaInterpolationFactor;
 
 						TrackProxy->CurrentPositionBufferIndex = 0;
-						TrackProxy->PositionBuffers[0].Update(bExtrapolateFrames ? ExtrapolatedPositions : MeshDataToUse->Positions);
-						TrackProxy->PositionBuffers[1].Update(MotionVectors);
+						TrackProxy->PositionBuffers[0].Update(RHICmdList, bExtrapolateFrames ? ExtrapolatedPositions : MeshDataToUse->Positions);
+						TrackProxy->PositionBuffers[1].Update(RHICmdList, MotionVectors);
 						TrackProxy->PositionBufferFrameIndices[0] = FrameIndexToUse;
 						TrackProxy->PositionBufferFrameIndices[1] = -1;
 						TrackProxy->PositionBufferFrameTimes[0] = Time;
@@ -1237,8 +1238,8 @@ void FGeometryCacheSceneProxy::FrameUpdate() const
 				{
 					check(false);
 					// Only update the size on this thread
-					TrackProxy->IndexBuffer.UpdateSizeOnly(TrackProxy->MeshData->Indices.Num());
-					TrackProxy->VertexBuffer.UpdateSizeTyped<FNoPositionVertex>(TrackProxy->MeshData->Vertices.Num());
+					TrackProxy->IndexBuffer.UpdateSizeOnly(RHICmdList, TrackProxy->MeshData->Indices.Num());
+					TrackProxy->VertexBuffer.UpdateSizeTyped<FNoPositionVertex>(RHICmdList, TrackProxy->MeshData->Vertices.Num());
 
 					// Do the interpolation on a worker thread
 					FGraphEventRef CompletionFence = FFunctionGraphTask::CreateAndDispatchWhenReady([]()
@@ -1247,8 +1248,7 @@ void FGeometryCacheSceneProxy::FrameUpdate() const
 					}, GET_STATID(STAT_BufferUpdateTask), NULL, ENamedThreads::AnyThread);
 
 					// Queue a command on the RHI thread that waits for the interpolation job and then uploads them to the GPU
-					FRHICommandListImmediate& RHICommandList = GetImmediateCommandList_ForRenderCommand();
-					new (RHICommandList.AllocCommand<FRHICommandUpdateGeometryCacheBuffer>())FRHICommandUpdateGeometryCacheBuffer(
+					new (RHICmdList.AllocCommand<FRHICommandUpdateGeometryCacheBuffer>())FRHICommandUpdateGeometryCacheBuffer(
 						CompletionFence,
 						TrackProxy->VertexBuffer.VertexBufferRHI,
 						TrackProxy->MeshData->Vertices.GetData(),
@@ -1294,7 +1294,7 @@ void FGeometryCacheSceneProxy::ClearSections()
 	Scratch.Empty();
 }
 
-void FGeomCacheTrackProxy::InitRenderResources(int32 NumVertices, int32 NumIndices)
+void FGeomCacheTrackProxy::InitRenderResources(FRHICommandListBase& RHICmdList, int32 NumVertices, int32 NumIndices)
 {
 	check(NumVertices);
 	check(NumIndices);
@@ -1316,17 +1316,17 @@ void FGeomCacheTrackProxy::InitRenderResources(int32 NumVertices, int32 NumIndic
 	IndexBuffer.NumValidIndices = 0;
 
 	// Init vertex factory
-	VertexFactory.Init(&PositionBuffers[0], &PositionBuffers[1], &TangentXBuffer, &TangentZBuffer, &TextureCoordinatesBuffer, &ColorBuffer);
+	VertexFactory.Init(RHICmdList, &PositionBuffers[0], &PositionBuffers[1], &TangentXBuffer, &TangentZBuffer, &TextureCoordinatesBuffer, &ColorBuffer);
 
 	// Enqueue initialization of render resource
-	BeginInitResource(&PositionBuffers[0]);
-	BeginInitResource(&PositionBuffers[1]);
-	BeginInitResource(&TangentXBuffer);
-	BeginInitResource(&TangentZBuffer);
-	BeginInitResource(&TextureCoordinatesBuffer);
-	BeginInitResource(&ColorBuffer);
-	BeginInitResource(&IndexBuffer);
-	BeginInitResource(&VertexFactory);
+	PositionBuffers[0].InitResource(RHICmdList);
+	PositionBuffers[1].InitResource(RHICmdList);
+	TangentXBuffer.InitResource(RHICmdList);
+	TangentZBuffer.InitResource(RHICmdList);
+	TextureCoordinatesBuffer.InitResource(RHICmdList);
+	ColorBuffer.InitResource(RHICmdList);
+	IndexBuffer.InitResource(RHICmdList);
+	VertexFactory.InitResource(RHICmdList);
 
 	bResourcesInitialized = true;
 }
@@ -1381,10 +1381,8 @@ FGeomCacheVertexFactory::FGeomCacheVertexFactory(ERHIFeatureLevel::Type InFeatur
 
 }
 
-void FGeomCacheVertexFactory::Init_RenderThread(const FVertexBuffer* PositionBuffer, const FVertexBuffer* MotionBlurDataBuffer, const FVertexBuffer* TangentXBuffer, const FVertexBuffer* TangentZBuffer, const FVertexBuffer* TextureCoordinateBuffer, const FVertexBuffer* ColorBuffer)
+void FGeomCacheVertexFactory::Init(FRHICommandListBase& RHICmdList, const FVertexBuffer* PositionBuffer, const FVertexBuffer* MotionBlurDataBuffer, const FVertexBuffer* TangentXBuffer, const FVertexBuffer* TangentZBuffer, const FVertexBuffer* TextureCoordinateBuffer, const FVertexBuffer* ColorBuffer)
 {
-	check(IsInRenderingThread());
-
 	// Initialize the vertex factory's stream components.
 	FDataType NewData;
 	NewData.PositionComponent = FVertexStreamComponent(PositionBuffer, 0, sizeof(FVector3f), VET_Float3);
@@ -1395,24 +1393,7 @@ void FGeomCacheVertexFactory::Init_RenderThread(const FVertexBuffer* PositionBuf
 	NewData.ColorComponent = FVertexStreamComponent(ColorBuffer, 0, sizeof(FColor), VET_Color);
 	NewData.MotionBlurDataComponent = FVertexStreamComponent(MotionBlurDataBuffer, 0, sizeof(FVector3f), VET_Float3);
 
-	SetData(NewData);
-}
-
-void FGeomCacheVertexFactory::Init(const FVertexBuffer* PositionBuffer, const FVertexBuffer* MotionBlurDataBuffer, const FVertexBuffer* TangentXBuffer, const FVertexBuffer* TangentZBuffer, const FVertexBuffer* TextureCoordinateBuffer, const FVertexBuffer* ColorBuffer)
-{
-	if (IsInRenderingThread())
-	{
-		Init_RenderThread(PositionBuffer, MotionBlurDataBuffer, TangentXBuffer, TangentZBuffer, TextureCoordinateBuffer, ColorBuffer);
-	}
-	else
-	{
-
-		ENQUEUE_RENDER_COMMAND(InitGeomCacheVertexFactory)(
-			[this, PositionBuffer, MotionBlurDataBuffer, TangentXBuffer, TangentZBuffer, TextureCoordinateBuffer, ColorBuffer](FRHICommandListImmediate& RHICmdList)
-		{
-			Init_RenderThread(PositionBuffer, MotionBlurDataBuffer, TangentXBuffer, TangentZBuffer, TextureCoordinateBuffer, ColorBuffer);
-		});
-	}
+	SetData(RHICmdList, NewData);
 }
 
 void FGeomCacheIndexBuffer::InitRHI(FRHICommandListBase& RHICmdList)
@@ -1433,12 +1414,10 @@ void FGeomCacheIndexBuffer::ReleaseRHI()
 	FIndexBuffer::ReleaseRHI();
 }
 
-void FGeomCacheIndexBuffer::Update(const TArray<uint32>& Indices)
+void FGeomCacheIndexBuffer::Update(FRHICommandListBase& RHICmdList, const TArray<uint32>& Indices)
 {
 	SCOPE_CYCLE_COUNTER(STAT_IndexBufferUpdate);
 
-	FRHICommandListBase& RHICmdList = FRHICommandListImmediate::Get();
-
 	void* Buffer = nullptr;
 
 	NumValidIndices = 0;
@@ -1483,10 +1462,8 @@ void FGeomCacheIndexBuffer::Update(const TArray<uint32>& Indices)
 	}
 }
 
-void FGeomCacheIndexBuffer::UpdateSizeOnly(int32 NewNumIndices)
+void FGeomCacheIndexBuffer::UpdateSizeOnly(FRHICommandListBase& RHICmdList, int32 NewNumIndices)
 {
-	FRHICommandListBase& RHICmdList = FRHICommandListImmediate::Get();
-
 	// We only ever grow in size. Ok for now?
 	bool bReallocate = false;
 	if (NewNumIndices > NumAllocatedIndices)
@@ -1543,12 +1520,11 @@ void FGeomCacheColorBuffer::InitRHI(FRHICommandListBase& RHICmdList)
 	}
 }
 
-void FGeomCacheVertexBuffer::UpdateRaw(const void* Data, int32 NumItems, int32 ItemSizeBytes, int32 ItemStrideBytes)
+void FGeomCacheVertexBuffer::UpdateRaw(FRHICommandListBase& RHICmdList, const void* Data, int32 NumItems, int32 ItemSizeBytes, int32 ItemStrideBytes)
 {
 	SCOPE_CYCLE_COUNTER(STAT_VertexBufferUpdate);
 	int32 NewSizeInBytes = ItemSizeBytes * NumItems;
 	bool bCanMemcopy = ItemSizeBytes == ItemStrideBytes;
-	FRHICommandListBase& RHICmdList = FRHICommandListImmediate::Get();
 
 	bool bReallocate = false;
 	if (NewSizeInBytes > SizeInBytes)
@@ -1585,10 +1561,8 @@ void FGeomCacheVertexBuffer::UpdateRaw(const void* Data, int32 NumItems, int32 I
 	}
 }
 
-void FGeomCacheVertexBuffer::UpdateSize(int32 NewSizeInBytes)
+void FGeomCacheVertexBuffer::UpdateSize(FRHICommandListBase& RHICmdList, int32 NewSizeInBytes)
 {
-	FRHICommandListBase& RHICmdList = FRHICommandListImmediate::Get();
-
 	bool bReallocate = false;
 	if (NewSizeInBytes > SizeInBytes)
 	{
diff --git a/Engine/Plugins/Runtime/GeometryCache/Source/GeometryCache/Public/GeometryCacheSceneProxy.h b/Engine/Plugins/Runtime/GeometryCache/Source/GeometryCache/Public/GeometryCacheSceneProxy.h
--- a/Engine/Plugins/Runtime/GeometryCache/Source/GeometryCache/Public/GeometryCacheSceneProxy.h
+++ b/Engine/Plugins/Runtime/GeometryCache/Source/GeometryCache/Public/GeometryCacheSceneProxy.h
@@ -37,20 +37,20 @@ public:
 	/**
 	 * Sugar function to update from a typed array.
 	 */
-	template<class DataType> void Update(const TArray<DataType>& Vertices)
+	template<class DataType> void Update(FRHICommandListBase& RHICmdList, const TArray<DataType>& Vertices)
 	{
 		int32 InSize = Vertices.Num() * sizeof(DataType);
-		UpdateRaw(Vertices.GetData(), InSize, 1, 1);
+		UpdateRaw(RHICmdList, Vertices.GetData(), InSize, 1, 1);
 	}
 
-	void UpdatePositionsOnly(const TArray<FDynamicMeshVertex>& Vertices)
+	void UpdatePositionsOnly(FRHICommandListBase& RHICmdList, const TArray<FDynamicMeshVertex>& Vertices)
 	{
 		const uint32 PositionOffset = STRUCT_OFFSET(FDynamicMeshVertex, Position);
 		const uint32 PositionSize = sizeof(((FDynamicMeshVertex*)nullptr)->Position);
-		UpdateRaw(Vertices.GetData() + PositionOffset, Vertices.Num(), PositionSize, sizeof(FDynamicMeshVertex));
+		UpdateRaw(RHICmdList, Vertices.GetData() + PositionOffset, Vertices.Num(), PositionSize, sizeof(FDynamicMeshVertex));
 	}
 
-	void UpdateExceptPositions(const TArray<FDynamicMeshVertex>& Vertices)
+	void UpdateExceptPositions(FRHICommandListBase& RHICmdList, const TArray<FDynamicMeshVertex>& Vertices)
 	{
 		const uint32 PositionSize = sizeof(((FDynamicMeshVertex*)nullptr)->Position);
 		const uint32 PositionOffset = STRUCT_OFFSET(FDynamicMeshVertex, Position);
@@ -58,25 +58,25 @@ public:
 		static_assert(PositionOffset == 0, "Expecting position to be the first struct member");
 		static_assert(PositionSize == STRUCT_OFFSET(FDynamicMeshVertex, TextureCoordinate), "Expecting the texture coordinate to immediately follow the Position");
 
-		UpdateRaw((int8*)Vertices.GetData() + PositionSize, Vertices.Num(), sizeof(FDynamicMeshVertex) - PositionSize, sizeof(FDynamicMeshVertex));
+		UpdateRaw(RHICmdList, (int8*)Vertices.GetData() + PositionSize, Vertices.Num(), sizeof(FDynamicMeshVertex) - PositionSize, sizeof(FDynamicMeshVertex));
 	}
 
 	/**
 	 * Update the raw contents of the buffer, possibly reallocate if needed.
 	 */
-	void UpdateRaw(const void* Data, int32 NumItems, int32 ItemSizeBytes, int32 ItemStrideBytes);
+	void UpdateRaw(FRHICommandListBase& RHICmdList, const void* Data, int32 NumItems, int32 ItemSizeBytes, int32 ItemStrideBytes);
 
 	/**
 	 * Resize the buffer but don't initialize it with any data.
 	 */
-	void UpdateSize(int32 NewSizeInBytes);
+	void UpdateSize(FRHICommandListBase& RHICmdList, int32 NewSizeInBytes);
 
 	/**
 	* Resize the buffer but don't initialize it with any data.
 	*/
-	template<class DataType> void UpdateSizeTyped(int32 NewSizeInElements)
+	template<class DataType> void UpdateSizeTyped(FRHICommandListBase& RHICmdList, int32 NewSizeInElements)
 	{
-		UpdateSize(sizeof(DataType) * NewSizeInElements);
+		UpdateSize(RHICmdList, sizeof(DataType) * NewSizeInElements);
 	}
 
 	/**
@@ -120,9 +120,8 @@ public:
 	/**
 		Update the data and possibly reallocate if needed.
 	*/
-	void Update(const TArray<uint32>& Indices);
-
-	void UpdateSizeOnly(int32 NewNumIndices);
+	void Update(FRHICommandListBase& RHICmdList, const TArray<uint32>& Indices);
+	void UpdateSizeOnly(FRHICommandListBase& RHICmdList, int32 NewNumIndices);
 
 	unsigned SizeInBytes() { return NumAllocatedIndices * sizeof(uint32); }
 
@@ -136,14 +135,9 @@ protected:
 class GEOMETRYCACHE_API FGeomCacheVertexFactory : public FGeometryCacheVertexVertexFactory
 {
 public:
-
 	FGeomCacheVertexFactory(ERHIFeatureLevel::Type InFeatureLevel);
 
-	/** Init function that should only be called on render thread. */
-	void Init_RenderThread(const FVertexBuffer* PositionBuffer, const FVertexBuffer* MotionBlurDataBuffer, const FVertexBuffer* TangentXBuffer, const FVertexBuffer* TangentZBuffer, const FVertexBuffer* TextureCoordinateBuffer, const FVertexBuffer* ColorBuffer);
-
-	/** Init function that can be called on any thread, and will do the right thing (enqueue command if called on main thread) */
-	void Init(const FVertexBuffer* PositionBuffer, const FVertexBuffer* MotionBlurDataBuffer, const FVertexBuffer* TangentXBuffer, const FVertexBuffer* TangentZBuffer, const FVertexBuffer* TextureCoordinateBuffer, const FVertexBuffer* ColorBuffer);
+	void Init(FRHICommandListBase& RHICmdList, const FVertexBuffer* PositionBuffer, const FVertexBuffer* MotionBlurDataBuffer, const FVertexBuffer* TangentXBuffer, const FVertexBuffer* TangentZBuffer, const FVertexBuffer* TextureCoordinateBuffer, const FVertexBuffer* ColorBuffer);
 };
 
 /**
@@ -224,7 +218,7 @@ public:
 	 * @param NumVertices - The initial number of vertices to initialize the buffers with. Must be greater than 0
 	 * @param NumIndices - The initial number of indices to initialize the buffers with. Must be greater than 0
 	 */
-	virtual void InitRenderResources(int32 NumVertices, int32 NumIndices);
+	virtual void InitRenderResources(FRHICommandListBase& RHICmdList, int32 NumVertices, int32 NumIndices);
 
 	/** MeshData storing information used for rendering this Track */
 	FGeometryCacheMeshData* MeshData;
@@ -300,7 +294,7 @@ public:
 	uint32 GetAllocatedSize(void) const;
 	// End FPrimitiveSceneProxy interface.
 
-	void UpdateAnimation(float NewTime, bool bLooping, bool bIsPlayingBackwards, float PlaybackSpeed, float MotionVectorScale);
+	void UpdateAnimation(FRHICommandListBase& RHICmdList, float NewTime, bool bLooping, bool bIsPlayingBackwards, float PlaybackSpeed, float MotionVectorScale);
 
 	/** Update world matrix for specific section */
 	void UpdateSectionWorldMatrix(const int32 SectionIndex, const FMatrix& WorldMatrix);
@@ -321,7 +315,7 @@ public:
 	const TArray<FGeomCacheTrackProxy*>& GetTracks() const { return Tracks; }
 
 private:
-	void FrameUpdate() const;
+	void FrameUpdate(FRHICommandListBase& RHICmdList) const;
 
 	void CreateMeshBatch(
 		FRHICommandListBase& RHICmdList,
diff --git a/Engine/Plugins/Runtime/HairStrands/Source/HairStrandsCore/Private/GroomComponent.cpp b/Engine/Plugins/Runtime/HairStrands/Source/HairStrandsCore/Private/GroomComponent.cpp
--- a/Engine/Plugins/Runtime/HairStrands/Source/HairStrandsCore/Private/GroomComponent.cpp
+++ b/Engine/Plugins/Runtime/HairStrands/Source/HairStrandsCore/Private/GroomComponent.cpp
@@ -550,6 +550,9 @@ public:
 				}
 			}
 		}
+
+		// TODO: Rendering hair strands currently requires a deterministic ordering of calls to GetDynamicMeshElements.
+		bSupportsParallelGDME = false;
 	}
 
 	virtual ~FHairStrandsSceneProxy()
diff --git a/Engine/Source/Runtime/Engine/Private/GeometryCacheVertexFactory.cpp b/Engine/Source/Runtime/Engine/Private/GeometryCacheVertexFactory.cpp
--- a/Engine/Source/Runtime/Engine/Private/GeometryCacheVertexFactory.cpp
+++ b/Engine/Source/Runtime/Engine/Private/GeometryCacheVertexFactory.cpp
@@ -98,6 +98,11 @@ void FGeometryCacheVertexVertexFactory::ModifyCompilationEnvironment(const FVert
 }
 
 void FGeometryCacheVertexVertexFactory::SetData(const FDataType& InData)
+{
+	SetData(FRHICommandListImmediate::Get(), InData);
+}
+
+void FGeometryCacheVertexVertexFactory::SetData(FRHICommandListBase& RHICmdList, const FDataType& InData)
 {
 	// The shader code makes assumptions that the color component is a FColor, performing swizzles on ES3 and Metal platforms as necessary
 	// If the color is sent down as anything other than VET_Color then you'll get an undesired swizzle on those platforms
@@ -105,7 +110,7 @@ void FGeometryCacheVertexVertexFactory::SetData(const FDataType& InData)
 
 	Data = InData;
 	// This will call InitRHI below where the real action happens
-	UpdateRHI(FRHICommandListImmediate::Get());
+	UpdateRHI(RHICmdList);
 }
 
 class FDefaultGeometryCacheVertexBuffer : public FVertexBuffer
diff --git a/Engine/Source/Runtime/Engine/Private/GroomVisualizationData.cpp b/Engine/Source/Runtime/Engine/Private/GroomVisualizationData.cpp
--- a/Engine/Source/Runtime/Engine/Private/GroomVisualizationData.cpp
+++ b/Engine/Source/Runtime/Engine/Private/GroomVisualizationData.cpp
@@ -192,6 +192,7 @@ static void AddVisualizationMode(
 
 void FGroomVisualizationData::Initialize()
 {
+	UE::TScopeLock Lock(Mutex);
 	if (!bIsInitialized && IsGroomEnabled())
 	{
 		AddVisualizationMode(ModeMap, true, EGroomViewMode::None,						LOCTEXT("NoneDebug", "None"),								LOCTEXT("NoneDebugDesc", "No debug mode"));
diff --git a/Engine/Source/Runtime/Engine/Private/Particles/ParticleMemoryStatManager.cpp b/Engine/Source/Runtime/Engine/Private/Particles/ParticleMemoryStatManager.cpp
--- a/Engine/Source/Runtime/Engine/Private/Particles/ParticleMemoryStatManager.cpp
+++ b/Engine/Source/Runtime/Engine/Private/Particles/ParticleMemoryStatManager.cpp
@@ -130,8 +130,6 @@ void FParticleMemoryStatManager::DumpParticleMemoryStats(FOutputDevice& Ar)
 	Ar.Logf(TEXT("Max wasted GT,%d"), DynamicEmitterGTMem_Waste_MAX);
 	Ar.Logf(TEXT("Largest single GT allocation,%d"), DynamicEmitterGTMem_Largest_MAX);
 	Ar.Logf(TEXT("Largest single RT allocation,%d"), LargestRTParticleData_MAX);
-
-	GParticleOrderPool.DumpInfo(Ar);
 }
 
 /**
diff --git a/Engine/Source/Runtime/Engine/Private/Particles/ParticleSystemRender.cpp b/Engine/Source/Runtime/Engine/Private/Particles/ParticleSystemRender.cpp
--- a/Engine/Source/Runtime/Engine/Private/Particles/ParticleSystemRender.cpp
+++ b/Engine/Source/Runtime/Engine/Private/Particles/ParticleSystemRender.cpp
@@ -35,6 +35,7 @@
 #include "RenderCore.h"
 #include "DataDrivenShaderPlatformInfo.h"
 #include "PrimitiveUniformShaderParametersBuilder.h"
+#include "Experimental/ConcurrentLinearAllocator.h"
 
 DECLARE_CYCLE_STAT(TEXT("ParticleSystemSceneProxy Create GT"), STAT_FParticleSystemSceneProxy_Create, STATGROUP_Particles);
 DECLARE_CYCLE_STAT(TEXT("ParticleSystemSceneProxy GetMeshElements RT"), STAT_FParticleSystemSceneProxy_GetMeshElements, STATGROUP_Particles);
@@ -100,9 +101,6 @@ float GMinParticleDrawTimeToTrack = .0001f;
 /** Whether to do LOD calculation on GameThread in game */
 extern bool GbEnableGameThreadLODCalculation;
 
-///////////////////////////////////////////////////////////////////////////////
-FParticleOrderPool GParticleOrderPool;
-
 ///////////////////////////////////////////////////////////////////////////////
 
 /**
@@ -927,11 +925,11 @@ void FDynamicSpriteEmitterData::GetDynamicMeshElementsEmitter(const FParticleSys
 				{
 					SCOPE_CYCLE_COUNTER(STAT_FDynamicSpriteEmitterData_PerParticleWorkOrTasks);
 					{
-						FParticleOrder* ParticleOrder = NULL;
+						FParticleOrder* ParticleOrder = nullptr;
 
 						if (bSort)
 						{
-							ParticleOrder = GParticleOrderPool.GetParticleOrderData(ParticleCount);
+							ParticleOrder = (FParticleOrder*)FConcurrentLinearAllocator::Malloc(ParticleCount, 16);
 							SortSpriteParticles(SourceData->SortMode, SourceData->bUseLocalSpace, SourceData->ActiveParticleCount, 
 								SourceData->DataContainer.ParticleData, SourceData->ParticleStride, SourceData->DataContainer.ParticleIndices,
 								View, Proxy->GetLocalToWorld(), ParticleOrder);
@@ -939,6 +937,7 @@ void FDynamicSpriteEmitterData::GetDynamicMeshElementsEmitter(const FParticleSys
 
 						// Fill vertex buffers.
 						GetVertexAndIndexData(Allocation.Buffer, DynamicParameterAllocation.Buffer, NULL, ParticleOrder, View->ViewMatrices.GetViewOrigin(), Proxy->GetLocalToWorld(), InstanceFactor);
+						FConcurrentLinearAllocator::Free(ParticleOrder);
 					}
 				}
 
diff --git a/Engine/Source/Runtime/Engine/Private/PrimitiveSceneProxy.cpp b/Engine/Source/Runtime/Engine/Private/PrimitiveSceneProxy.cpp
--- a/Engine/Source/Runtime/Engine/Private/PrimitiveSceneProxy.cpp
+++ b/Engine/Source/Runtime/Engine/Private/PrimitiveSceneProxy.cpp
@@ -106,7 +106,7 @@ static TAutoConsoleVariable<int32> CVarApproximateOcclusionQueries(
 	ECVF_RenderThreadSafe
 );
 
-bool GParallelGatherDynamicMeshElements = false;
+bool GParallelGatherDynamicMeshElements = true;
 static FAutoConsoleVariableRef CVarParallelGatherDynamicMeshElements(
 	TEXT("r.Visibility.DynamicMeshElements.Parallel"),
 	GParallelGatherDynamicMeshElements,
diff --git a/Engine/Source/Runtime/Engine/Private/StaticMeshRender.cpp b/Engine/Source/Runtime/Engine/Private/StaticMeshRender.cpp
--- a/Engine/Source/Runtime/Engine/Private/StaticMeshRender.cpp
+++ b/Engine/Source/Runtime/Engine/Private/StaticMeshRender.cpp
@@ -1473,7 +1473,6 @@ void FStaticMeshSceneProxy::GetMeshDescription(int32 LODIndex, TArray<FMeshBatch
 void FStaticMeshSceneProxy::GetDynamicMeshElements(const TArray<const FSceneView*>& Views, const FSceneViewFamily& ViewFamily, uint32 VisibilityMap, FMeshElementCollector& Collector) const
 {
 	QUICK_SCOPE_CYCLE_COUNTER(STAT_StaticMeshSceneProxy_GetMeshElements);
-	checkSlow(IsInRenderingThread());
 
 	const bool bIsLightmapSettingError = HasStaticLighting() && !HasValidSettingsForStaticLighting();
 	const bool bProxyIsSelected = IsSelected();
diff --git a/Engine/Source/Runtime/Engine/Public/GeometryCacheVertexFactory.h b/Engine/Source/Runtime/Engine/Public/GeometryCacheVertexFactory.h
--- a/Engine/Source/Runtime/Engine/Public/GeometryCacheVertexFactory.h
+++ b/Engine/Source/Runtime/Engine/Public/GeometryCacheVertexFactory.h
@@ -105,8 +105,11 @@ public:
 	/**
 	* An implementation of the interface used by TSynchronizedResource to update the resource with new data from the game thread.
 	*/
+	UE_DEPRECATED(5.4, "SetData requires an RHI command list.")
 	ENGINE_API void SetData(const FDataType& InData);
 
+	ENGINE_API void SetData(FRHICommandListBase& RHICmdList, const FDataType& InData);
+
 	UE_DEPRECATED(5.3, "CreateManualVertexFetchUniformBuffer requires a command list.")
 	ENGINE_API void CreateManualVertexFetchUniformBuffer(
 		const FVertexBuffer* PoistionBuffer,
diff --git a/Engine/Source/Runtime/Engine/Public/GroomVisualizationData.h b/Engine/Source/Runtime/Engine/Public/GroomVisualizationData.h
--- a/Engine/Source/Runtime/Engine/Public/GroomVisualizationData.h
+++ b/Engine/Source/Runtime/Engine/Public/GroomVisualizationData.h
@@ -3,6 +3,7 @@
 #pragma once
 
 #include "CoreMinimal.h"
+#include "Async/Mutex.h"
 
 class FSceneView;
 
@@ -107,7 +108,10 @@ private:
 	FString ConsoleDocumentationVisualizationMode;
 
 	/** Flag indicating if system is initialized. **/
-	bool bIsInitialized;
+	std::atomic_bool bIsInitialized;
+
+	/** Mutex for initialization. */
+	UE::FMutex Mutex;
 };
 
 ENGINE_API FGroomVisualizationData& GetGroomVisualizationData();
diff --git a/Engine/Source/Runtime/Engine/Public/ParticleHelper.h b/Engine/Source/Runtime/Engine/Public/ParticleHelper.h
--- a/Engine/Source/Runtime/Engine/Public/ParticleHelper.h
+++ b/Engine/Source/Runtime/Engine/Public/ParticleHelper.h
@@ -1118,64 +1118,6 @@ struct FAsyncBufferFillData
 -----------------------------------------------------------------------------*/
 class FParticleVertexFactoryBase;
 
-/*-----------------------------------------------------------------------------
-	Particle order helper class
------------------------------------------------------------------------------*/
-class FParticleOrderPool
-{
-public:
-	FParticleOrderPool() :
-		  ParticleOrder(NULL)
-		, CurrentSize(0)
-		, MaxSize(0)
-	{
-	}
-
-	~FParticleOrderPool()
-	{
-		FreePool();
-	}
-
-	FParticleOrder* GetParticleOrderData(int32 InCount, bool bZeroMem = false)
-	{
-		if (InCount > MaxSize)
-		{
-			MaxSize = FMath::Max<int32>(InCount, 64);
-			ParticleOrder = (FParticleOrder*)FMemory::Realloc(ParticleOrder, MaxSize * sizeof(FParticleOrder));
-			check(ParticleOrder);
-			if (bZeroMem == true)
-			{
-				FMemory::Memzero(ParticleOrder, MaxSize * sizeof(FParticleOrder));
-			}
-		}
-		CurrentSize = InCount;
-		return ParticleOrder;
-	}
-
-	void FreePool()
-	{
-		FMemory::Free(ParticleOrder);
-		ParticleOrder = NULL;
-		CurrentSize = 0;
-		MaxSize = 0;
-	}
-
-#if STATS
-	void DumpInfo(FOutputDevice& Ar)
-	{
-		Ar.Logf(TEXT("Particle Order Pool Stats"));
-		Ar.Logf(TEXT("%5d entries for %5d bytes"), MaxSize, MaxSize * sizeof(FParticleOrder));
-	}
-#endif
-
-protected:
-	FParticleOrder* ParticleOrder;
-	int32 CurrentSize;
-	int32 MaxSize;
-};
-
-extern FParticleOrderPool GParticleOrderPool;
-
 /*-----------------------------------------------------------------------------
 	Particle Dynamic Data
 -----------------------------------------------------------------------------*/
diff --git a/Engine/Source/Runtime/Renderer/Private/SceneVisibility.cpp b/Engine/Source/Runtime/Renderer/Private/SceneVisibility.cpp
--- a/Engine/Source/Runtime/Renderer/Private/SceneVisibility.cpp
+++ b/Engine/Source/Runtime/Renderer/Private/SceneVisibility.cpp
@@ -2954,6 +2954,7 @@ void FGPUOcclusionParallel::AddPrimitives(FPrimitiveRange PrimitiveRange)
 			{
 				ViewPacket.Relevance.CommandPipe.AddNumCommands(1);
 				ViewPacket.Relevance.CommandPipe.EnqueueCommand(MoveTemp(NonOccludedPrimitives));
+				NonOccludedPrimitives.Reset();
 				NonOccludedPrimitives.Reserve(MaxNonOccludedPrimitives);
 			}
 		}
@@ -2971,6 +2972,7 @@ void FGPUOcclusionParallel::Finish(UE::Tasks::FTaskEvent& OcclusionCullTasks)
 	{
 		ViewPacket.Relevance.CommandPipe.AddNumCommands(1);
 		ViewPacket.Relevance.CommandPipe.EnqueueCommand(MoveTemp(NonOccludedPrimitives));
+		NonOccludedPrimitives.Reset();
 	}
 
 	for (FGPUOcclusionParallelPacket* Packet : Packets)
@@ -3984,9 +3986,7 @@ void FVisibilityTaskData::LaunchVisibilityTasks()
 	#endif
 	}
 
-	const int32 NumAsyncDynamicMeshElementContexts = TaskConfig.Schedule == EVisibilityTaskSchedule::Parallel ? GetNumDynamicMeshElementTasks() : 0;
-
-	DynamicMeshElements.ContextContainer.Init(SceneRenderer, NumAsyncDynamicMeshElementContexts);
+	DynamicMeshElements.ContextContainer.Init(SceneRenderer, GetNumDynamicMeshElementTasks());
 	DynamicMeshElements.ViewCommandsPerView.SetNum(Views.Num());
 
 	Tasks.LightVisibility.AddPrerequisites(UE::Tasks::Launch(UE_SOURCE_LOCATION, [this]
