GeometryCache: Fixed uninitialized bounding box in GeometryCacheMeshData that could cause issues for GeometryCacheComponent bounds.

diff --git a/Engine/Plugins/Runtime/GeometryCache/Source/GeometryCache/Classes/GeometryCacheMeshData.h b/Engine/Plugins/Runtime/GeometryCache/Source/GeometryCache/Classes/GeometryCacheMeshData.h
--- a/Engine/Plugins/Runtime/GeometryCache/Source/GeometryCache/Classes/GeometryCacheMeshData.h
+++ b/Engine/Plugins/Runtime/GeometryCache/Source/GeometryCache/Classes/GeometryCacheMeshData.h
@@ -96,7 +96,11 @@ struct GEOMETRYCACHE_API FGeometryCacheMeshData
 {
 	GENERATED_USTRUCT_BODY()
 
-	FGeometryCacheMeshData() : Hash(0) {}
+	FGeometryCacheMeshData()
+	: BoundingBox(EForceInit::ForceInitToZero)
+	, Hash(0)
+	{}
+
 	~FGeometryCacheMeshData()
 	{
 		Positions.Empty();
