GeometryCache: Fixed truncation warnings

#changelist validated
#virtualized

diff --git a/Engine/Plugins/FX/Niagara/Source/Niagara/Classes/NiagaraDataSetCompiledData.h b/Engine/Plugins/FX/Niagara/Source/Niagara/Classes/NiagaraDataSetCompiledData.h
--- a/Engine/Plugins/FX/Niagara/Source/Niagara/Classes/NiagaraDataSetCompiledData.h
+++ b/Engine/Plugins/FX/Niagara/Source/Niagara/Classes/NiagaraDataSetCompiledData.h
@@ -28,9 +28,9 @@ public:
 	uint32 GetInt32ComponentStart() const { return Int32ComponentStart; }
 	uint32 GetHalfComponentStart() const { return HalfComponentStart; }
 
-	void SetFloatComponentStart(uint32 Offset) { check(Offset <= TNumericLimits<uint16>::Max()); FloatComponentStart = Offset; }
-	void SetInt32ComponentStart(uint32 Offset) { check(Offset <= TNumericLimits<uint16>::Max()); Int32ComponentStart = Offset; }
-	void SetHalfComponentStart(uint32 Offset) { check(Offset <= TNumericLimits<uint16>::Max()); HalfComponentStart = Offset; }
+	void SetFloatComponentStart(uint32 Offset) { check(Offset <= TNumericLimits<uint16>::Max()); FloatComponentStart = static_cast<uint16>(Offset); }
+	void SetInt32ComponentStart(uint32 Offset) { check(Offset <= TNumericLimits<uint16>::Max()); Int32ComponentStart = static_cast<uint16>(Offset); }
+	void SetHalfComponentStart(uint32 Offset) { check(Offset <= TNumericLimits<uint16>::Max()); HalfComponentStart = static_cast<uint16>(Offset); }
 
 	uint32 GetNumFloatComponents() const { return LayoutInfo.GetNumFloatComponents(); }
 	uint32 GetNumInt32Components() const { return LayoutInfo.GetNumInt32Components(); }
diff --git a/Engine/Plugins/FX/Niagara/Source/Niagara/Classes/NiagaraSystem.h b/Engine/Plugins/FX/Niagara/Source/Niagara/Classes/NiagaraSystem.h
--- a/Engine/Plugins/FX/Niagara/Source/Niagara/Classes/NiagaraSystem.h
+++ b/Engine/Plugins/FX/Niagara/Source/Niagara/Classes/NiagaraSystem.h
@@ -955,7 +955,7 @@ protected:
 	The max number of substeps per frame can be set via fx.Niagara.SystemSimulation.MaxTickSubsteps
 	*/
 	UPROPERTY(EditAnywhere, Category = "System", meta = (EditCondition = "bFixedTickDelta", ForceUnits=s))
-	float FixedTickDeltaTime = 0.01667;
+	float FixedTickDeltaTime = 0.01667f;
 
 #if WITH_EDITORONLY_DATA
 	/** Settings used inside the baker */
diff --git a/Engine/Plugins/FX/Niagara/Source/Niagara/Public/NiagaraDebuggerCommon.h b/Engine/Plugins/FX/Niagara/Source/Niagara/Public/NiagaraDebuggerCommon.h
--- a/Engine/Plugins/FX/Niagara/Source/Niagara/Public/NiagaraDebuggerCommon.h
+++ b/Engine/Plugins/FX/Niagara/Source/Niagara/Public/NiagaraDebuggerCommon.h
@@ -724,7 +724,7 @@ struct FNiagaraDebugHUDSettingsData
 
 	// In world text if an error is detected
 	UPROPERTY(EditAnywhere, Category = "Colors")
-	FLinearColor InWorldErrorTextColor = FLinearColor(1.0f, 0.4, 0.3, 1.0f);
+	FLinearColor InWorldErrorTextColor = FLinearColor(1.0f, 0.4f, 0.3f, 1.0f);
 	// In world text color
 	UPROPERTY(EditAnywhere, Category = "Colors")
 	FLinearColor InWorldTextColor = FLinearColor::White;
@@ -735,11 +735,11 @@ struct FNiagaraDebugHUDSettingsData
 
 	// Message display warning text color
 	UPROPERTY(EditAnywhere, Category = "Colors")
-	FLinearColor MessageWarningTextColor = FLinearColor(0.9f, 0.7f, 0.0, 1.0f);
+	FLinearColor MessageWarningTextColor = FLinearColor(0.9f, 0.7f, 0.0f, 1.0f);
 
 	// Message display error text color
 	UPROPERTY(EditAnywhere, Category = "Colors")
-	FLinearColor MessageErrorTextColor = FLinearColor(1.0f, 0.4, 0.3, 1.0f);
+	FLinearColor MessageErrorTextColor = FLinearColor(1.0f, 0.4f, 0.3f, 1.0f);
 	
 	/** Opacity of the system color background tile in overview table rows. */
 	UPROPERTY(Config, EditAnywhere, Category = "Colors")
diff --git a/Engine/Plugins/FX/Niagara/Source/Niagara/Public/NiagaraRendererProperties.h b/Engine/Plugins/FX/Niagara/Source/Niagara/Public/NiagaraRendererProperties.h
--- a/Engine/Plugins/FX/Niagara/Source/Niagara/Public/NiagaraRendererProperties.h
+++ b/Engine/Plugins/FX/Niagara/Source/Niagara/Public/NiagaraRendererProperties.h
@@ -107,9 +107,9 @@ struct FNiagaraRendererVariableInfo
 
 	FNiagaraRendererVariableInfo() {}
 	FNiagaraRendererVariableInfo(int32 InDataOffset, int32 InGPUBufferOffset, int32 InNumComponents, bool bInUpload, bool bInHalfType)
-		: DatasetOffset(InDataOffset)
-		, GPUBufferOffset(InGPUBufferOffset)
-		, NumComponents(InNumComponents)
+		: DatasetOffset(static_cast<uint16>(InDataOffset))
+		, GPUBufferOffset(static_cast<uint16>(InGPUBufferOffset))
+		, NumComponents(static_cast<uint16>(InNumComponents))
 		, bUpload(bInUpload)
 		, bHalfType(bInHalfType)
 	{
diff --git a/Engine/Plugins/Runtime/GeometryCache/Source/GeometryCache/GeometryCache.Build.cs b/Engine/Plugins/Runtime/GeometryCache/Source/GeometryCache/GeometryCache.Build.cs
--- a/Engine/Plugins/Runtime/GeometryCache/Source/GeometryCache/GeometryCache.Build.cs
+++ b/Engine/Plugins/Runtime/GeometryCache/Source/GeometryCache/GeometryCache.Build.cs
@@ -30,6 +30,8 @@ public class GeometryCache : ModuleRules
             DynamicallyLoadedModuleNames.Add("GeometryCacheEd");
             PrivateDependencyModuleNames.Add("MeshUtilitiesCommon");
             PrivateDependencyModuleNames.Add("UnrealEd");
-        }        
+        }
+
+		UnsafeTypeCastWarningLevel = WarningLevel.Error;
 	}
 }
diff --git a/Engine/Plugins/Runtime/GeometryCache/Source/GeometryCache/Private/CodecV1.cpp b/Engine/Plugins/Runtime/GeometryCache/Source/GeometryCache/Private/CodecV1.cpp
--- a/Engine/Plugins/Runtime/GeometryCache/Source/GeometryCache/Private/CodecV1.cpp
+++ b/Engine/Plugins/Runtime/GeometryCache/Source/GeometryCache/Private/CodecV1.cpp
@@ -98,9 +98,9 @@ public:
 	FQuantizer(float Range, int32 NumBits)
 	{		
 		int32 BinCount = (int)FMath::Pow(2.0f, NumBits);
-		BinSize = Range / BinCount;
-		HalfBinSize = Range / BinCount / 2.0f;
-		OneOverBinSize = BinCount / Range;
+		BinSize = Range / static_cast<float>(BinCount);
+		HalfBinSize = Range / static_cast<float>(BinCount) / 2.0f;
+		OneOverBinSize = static_cast<float>(BinCount) / Range;
 	}
 #if WITH_EDITOR
 	/** Quantize a value */
@@ -108,7 +108,7 @@ public:
 	{
 		// We compensate for energy loss around zero, e.g., Given a bin size 1, we map [-0.5,0.5[ -> 0, [-1.5,-0.5[ -> -1, [0.5,1.5[ -> 1, 
 		int32 Negative = (int)(Value >= 0.0f) * 2 - 1; // Positive: 1, negative: -1
-		int32 IntValue = (FMath::Abs(Value) + HalfBinSize) * OneOverBinSize;
+		int32 IntValue = static_cast<int32>((FMath::Abs(Value) + HalfBinSize) * OneOverBinSize);
 		return IntValue * Negative;
 	}
 #endif // WITH_EDITOR
@@ -378,10 +378,10 @@ void FCodecV1Encoder::EncodeNormalStream(const FPackedNormal* Stream, uint64 Ele
 		// Load data
 		FPackedNormal& NormalValue = *(FPackedNormal*)RawElementData;		
 		
-		int8 dx = NormalValue.Vector.X - x;
-		int8 dy = NormalValue.Vector.Y - y;
-		int8 dz = NormalValue.Vector.Z - z;
-		int8 dw = NormalValue.Vector.W - w;
+		int8 dx = static_cast<int8>(NormalValue.Vector.X - x);
+		int8 dy = static_cast<int8>(NormalValue.Vector.Y - y);
+		int8 dz = static_cast<int8>(NormalValue.Vector.Z - z);
+		int8 dw = static_cast<int8>(NormalValue.Vector.W - w);
 
 		// Write residual	
 		WriteSymbol(Table, (uint8)dx);
@@ -549,7 +549,7 @@ bool FCodecV1Encoder::EncodeFrameData(FMemoryWriter& Writer, const FGeometryCach
 	}
 
 	// Additional stats
-	Statistics.DurationMs = CodingTime.Get();
+	Statistics.DurationMs = static_cast<float>(CodingTime.Get());
 	Statistics.All.Quality = 0.0f;	
 	Statistics.NumVertices = MeshData.Positions.Num();
 	UE_LOG(LogGeoCaStreamingCodecV1, Log, TEXT("Compressed %u vertices, %u bytes to %u bytes in %.2f milliseconds (%.2f ratio), quantizer precision: %.2f units."), 
@@ -932,10 +932,10 @@ void FCodecV1Decoder::DecodeNormalStream(FHuffmanBitStreamReader& Reader, FPacke
 	{
 		// Read coded residual
 		Reader.Refill();
-		x += Table.DecodeNoRefill(Reader);	
-		y += Table.DecodeNoRefill(Reader);
-		z += Table.DecodeNoRefill(Reader);
-		w += Table.DecodeNoRefill(Reader);
+		x += static_cast<uint8>(Table.DecodeNoRefill(Reader));	
+		y += static_cast<uint8>(Table.DecodeNoRefill(Reader));
+		z += static_cast<uint8>(Table.DecodeNoRefill(Reader));
+		w += static_cast<uint8>(Table.DecodeNoRefill(Reader));
 		
 		FPackedNormal* Value = (FPackedNormal*)RawElementData;
 		Value->Vector.X = x;
@@ -970,10 +970,10 @@ void FCodecV1Decoder::DecodeColorStream(FHuffmanBitStreamReader& Reader, FColor*
 		}
 																			
 		FColor* Value = (FColor*)RawElementData; // Save result to our list
-		Value->R = QuantizedValue.X;
-		Value->G = QuantizedValue.Y;
-		Value->B = QuantizedValue.Z;
-		Value->A = QuantizedValue.W;
+		Value->R = static_cast<uint8>(QuantizedValue.X);
+		Value->G = static_cast<uint8>(QuantizedValue.Y);
+		Value->B = static_cast<uint8>(QuantizedValue.Z);
+		Value->A = static_cast<uint8>(QuantizedValue.W);
 	}
 }
 
@@ -1295,7 +1295,7 @@ bool FCodecV1Decoder::DecodeFrameData(FBufferReader& Reader, FGeometryCacheMeshD
 
 	if (CVarCodecDebug.GetValueOnAnyThread() == 1)
 	{
-		const float TimeFloat = DecodingTime.Get();
+		const float TimeFloat = static_cast<float>(DecodingTime.Get());
 		UE_LOG(LogGeoCaStreamingCodecV1, Log, TEXT("Decoded frame with %i vertices in %.2f milliseconds."), NumVertices, TimeFloat);
 	}
 	
diff --git a/Engine/Plugins/Runtime/GeometryCache/Source/GeometryCache/Private/CodecV1.h b/Engine/Plugins/Runtime/GeometryCache/Source/GeometryCache/Private/CodecV1.h
--- a/Engine/Plugins/Runtime/GeometryCache/Source/GeometryCache/Private/CodecV1.h
+++ b/Engine/Plugins/Runtime/GeometryCache/Source/GeometryCache/Private/CodecV1.h
@@ -30,7 +30,7 @@ public:
 	// Mean Squared Error
 	FORCEINLINE float ReadMSE()
 	{
-		FVector2f MSE(SumSquaredX / Num, SumSquaredY / Num);
+		FVector2f MSE(static_cast<float>(SumSquaredX / static_cast<double>(Num)), static_cast<float>(SumSquaredY / static_cast<double>(Num)));
 		return (MSE.X + MSE.Y) / 2.0f;
 	}
 
@@ -60,7 +60,7 @@ public:
 	// Mean Squared Error
 	FORCEINLINE float ReadMSE()
 	{
-		FVector3f MSE(SumSquaredX / Num, SumSquaredY / Num, SumSquaredZ / Num);
+		FVector3f MSE(static_cast<float>(SumSquaredX / static_cast<double>(Num)), static_cast<float>(SumSquaredY / static_cast<double>(Num)), static_cast<float>(SumSquaredZ / static_cast<double>(Num)));
 		return (MSE.X + MSE.Y + MSE.Z) / 3.0f;
 	}
 
@@ -174,7 +174,7 @@ struct FStreamEncodingStatistics
 
 	FStreamEncodingStatistics(uint32 SetCodedNumBytes, uint32 SetRawNumBytes, float SetQuality) : CodedNumBytes(SetCodedNumBytes), RawNumBytes(SetRawNumBytes), Quality(SetQuality)
 	{
-		CompressionRatio = (float)RawNumBytes / CodedNumBytes;
+		CompressionRatio = static_cast<float>(RawNumBytes) / static_cast<float>(CodedNumBytes);
 	}
 };
 
diff --git a/Engine/Plugins/Runtime/GeometryCache/Source/GeometryCache/Private/FloatCompression.cpp b/Engine/Plugins/Runtime/GeometryCache/Source/GeometryCache/Private/FloatCompression.cpp
--- a/Engine/Plugins/Runtime/GeometryCache/Source/GeometryCache/Private/FloatCompression.cpp
+++ b/Engine/Plugins/Runtime/GeometryCache/Source/GeometryCache/Private/FloatCompression.cpp
@@ -54,7 +54,7 @@ public:
 
 		uint32 bytes = Stream.GetNumBytes();
 		uint32 originalBytes = sizeof(float) * NumValues;
-		float rate = originalBytes / (float)bytes;
+		float rate = static_cast<float>(originalBytes) / static_cast<float>(bytes);
 	}
 
 	/**
@@ -287,8 +287,8 @@ void IntegerTest()
 
 	for (int32 Index = 0; Index < NumValues; Index++)
 	{
-		float f = Index / 1024.0f;
-		float sf = sin(f);
+		float f = static_cast<float>(Index) / 1024.0f;
+		float sf = static_cast<float>(sin(f));
 		IntValues[Index] = (uint32)((sf*0.5 + 0.5) * 3999);
 	}
 
@@ -347,7 +347,7 @@ void FloatTest()
 
 	for (int32 Index = 0; Index < NumValues; Index++)
 	{
-		Values[Index] = (float)sin(Index / 1024.0f);
+		Values[Index] = static_cast<float>(sin(static_cast<float>(Index) / 1024.0f));
 	}
 
 	// lossless mode
diff --git a/Engine/Plugins/Runtime/GeometryCache/Source/GeometryCache/Private/GeometryCacheCodecRaw.cpp b/Engine/Plugins/Runtime/GeometryCache/Source/GeometryCache/Private/GeometryCacheCodecRaw.cpp
--- a/Engine/Plugins/Runtime/GeometryCache/Source/GeometryCache/Private/GeometryCacheCodecRaw.cpp
+++ b/Engine/Plugins/Runtime/GeometryCache/Source/GeometryCache/Private/GeometryCacheCodecRaw.cpp
@@ -51,7 +51,7 @@ void UGeometryCacheCodecRaw::CodeFrame(const FGeometryCacheCodecEncodeArguments
 	// Finish up any previous chunks
 	if (EncoderData.CurrentChunkId >= 0)
 	{
-		(*EncoderData.AppendChunksTo)[EncoderData.CurrentChunkId].LastFrame = GetEncoderFrameNumer() - 1;
+		(*EncoderData.AppendChunksTo)[EncoderData.CurrentChunkId].LastFrame = static_cast<float>(GetEncoderFrameNumer() - 1);
 	}
 
 	// Create a new chunk for this coded frame
@@ -69,8 +69,8 @@ void UGeometryCacheCodecRaw::CodeFrame(const FGeometryCacheCodecEncodeArguments
 
 	// Update other fields, note that for the last time start == end this is valid
 	// as any time past end will use the last frame no matter what the interval
-	NewChunk->FirstFrame = GetEncoderFrameNumer();
-	NewChunk->LastFrame = GetEncoderFrameNumer();
+	NewChunk->FirstFrame = static_cast<float>(GetEncoderFrameNumer());
+	NewChunk->LastFrame = static_cast<float>(GetEncoderFrameNumer());
 }
 #endif
 
@@ -79,7 +79,7 @@ DECLARE_CYCLE_STAT(TEXT("Deserialize Mesh"), STAT_DeserializeMesh, STATGROUP_Geo
 bool FGeometryCacheCodecRenderStateRaw::DecodeSingleFrame(FGeometryCacheCodecDecodeArguments &Args)
 {
 	// We have a chunk per frame so chunk id's are just frame id's
-	check(Args.Chunks[Args.FrameIdentifier].FirstFrame == Args.FrameIdentifier);
+	check(Args.Chunks[Args.FrameIdentifier].FirstFrame == static_cast<float>(Args.FrameIdentifier));
 
 	uint32 DataSize = 0;
 	const uint8 *data = IGeometryCacheStreamingManager::Get().MapChunk(Args.Track.GetTrack(), Args.FrameIdentifier, &DataSize);
diff --git a/Engine/Plugins/Runtime/GeometryCache/Source/GeometryCache/Private/GeometryCacheCodecV1.cpp b/Engine/Plugins/Runtime/GeometryCache/Source/GeometryCache/Private/GeometryCacheCodecV1.cpp
--- a/Engine/Plugins/Runtime/GeometryCache/Source/GeometryCache/Private/GeometryCacheCodecV1.cpp
+++ b/Engine/Plugins/Runtime/GeometryCache/Source/GeometryCache/Private/GeometryCacheCodecV1.cpp
@@ -49,7 +49,7 @@ UGeometryCacheCodecV1::~UGeometryCacheCodecV1()
 bool UGeometryCacheCodecV1::DecodeSingleFrame(FGeometryCacheCodecDecodeArguments &Args)
 {
 	// We have a chunk per frame so chunk id's are just frame id's
-	check(Args.Chunks[Args.FrameIdentifier].FirstFrame == Args.FrameIdentifier);
+	check(Args.Chunks[Args.FrameIdentifier].FirstFrame == static_cast<float>(Args.FrameIdentifier));
 
 	uint32 DataSize = 0;
 	const uint8 *data = IGeometryCacheStreamingManager::Get().MapChunk(Args.Track.GetTrack(), Args.FrameIdentifier, &DataSize);
@@ -104,7 +104,7 @@ void UGeometryCacheCodecV1::CodeFrame(const FGeometryCacheCodecEncodeArguments &
 	// Finish up any previous chunks
 	if (EncoderData.CurrentChunkId >= 0)
 	{
-		(*EncoderData.AppendChunksTo)[EncoderData.CurrentChunkId].LastFrame = GetEncoderFrameNumer()-1;
+		(*EncoderData.AppendChunksTo)[EncoderData.CurrentChunkId].LastFrame = static_cast<float>(GetEncoderFrameNumer()-1);
 	}
 
 	// Create a new chunk for this coded frame
@@ -122,8 +122,8 @@ void UGeometryCacheCodecV1::CodeFrame(const FGeometryCacheCodecEncodeArguments &
 
 	// Update other fields, note that for the last time start == end this is valid
 	// as any time past end will use the last frame no matter what the interval
-	NewChunk->FirstFrame = GetEncoderFrameNumer();
-	NewChunk->LastFrame = GetEncoderFrameNumer();
+	NewChunk->FirstFrame = static_cast<float>(GetEncoderFrameNumer());
+	NewChunk->LastFrame = static_cast<float>(GetEncoderFrameNumer());
 }
 
 #endif
@@ -141,7 +141,7 @@ DECLARE_CYCLE_STAT(TEXT("Deserialize Mesh"), STAT_DeserializeMeshV1, STATGROUP_G
 bool FGeometryCacheCodecRenderStateV1::DecodeSingleFrame(FGeometryCacheCodecDecodeArguments &Args)
 {
 	// We have a chunk per frame so chunk id's are just frame id's
-	check(Args.Chunks[Args.FrameIdentifier].FirstFrame == Args.FrameIdentifier);
+	check(Args.Chunks[Args.FrameIdentifier].FirstFrame == static_cast<float>(Args.FrameIdentifier));
 
 	uint32 DataSize = 0;
 	const uint8 *data = IGeometryCacheStreamingManager::Get().MapChunk(Args.Track.GetTrack(), Args.FrameIdentifier, &DataSize);
diff --git a/Engine/Plugins/Runtime/GeometryCache/Source/GeometryCache/Private/GeometryCacheComponent.cpp b/Engine/Plugins/Runtime/GeometryCache/Source/GeometryCache/Private/GeometryCacheComponent.cpp
--- a/Engine/Plugins/Runtime/GeometryCache/Source/GeometryCache/Private/GeometryCacheComponent.cpp
+++ b/Engine/Plugins/Runtime/GeometryCache/Source/GeometryCache/Private/GeometryCacheComponent.cpp
@@ -533,7 +533,7 @@ float UGeometryCacheComponent::GetTimeAtFrame(const int32 Frame) const
 	const float FrameTime = GetNumberOfFrames() > 1 ? Duration / (float)(GetNumberOfFrames() - 1) : 0.0f;
 	const int32 StartFrame = GeometryCache != nullptr ? GeometryCache->GetStartFrame() : 0;
 
-	return FMath::Clamp(FrameTime * (Frame - StartFrame), 0.0f, Duration);
+	return FMath::Clamp(FrameTime * static_cast<float>(Frame - StartFrame), 0.0f, Duration);
 }
 
 int32 UGeometryCacheComponent::GetNumberOfFrames() const
diff --git a/Engine/Plugins/Runtime/GeometryCache/Source/GeometryCache/Private/GeometryCacheConstantTopologyWriter.cpp b/Engine/Plugins/Runtime/GeometryCache/Source/GeometryCache/Private/GeometryCacheConstantTopologyWriter.cpp
--- a/Engine/Plugins/Runtime/GeometryCache/Source/GeometryCache/Private/GeometryCacheConstantTopologyWriter.cpp
+++ b/Engine/Plugins/Runtime/GeometryCache/Source/GeometryCache/Private/GeometryCacheConstantTopologyWriter.cpp
@@ -346,7 +346,7 @@ namespace UE::GeometryCacheHelpers
 		int32 MaxRecordedFrame = -1;
 		for (int32 Frame = 0; Frame < NumFrames; ++Frame)
 		{
-			const float Time = Frame / ConfigRef.FPS;
+			const float Time = static_cast<float>(Frame) / ConfigRef.FPS;
 			FFrameData& FrameData = FramesToMoveFrom[Frame];
 			if (FrameData.Positions.Num() != NumVertices)
 			{
@@ -418,7 +418,7 @@ namespace UE::GeometryCacheHelpers
 		}
 		check(MaxRecordedFrame >= 0);
 		TArray<FMatrix> Mats { FMatrix::Identity, FMatrix::Identity };
-		TArray<float> MatTimes { 0.0f, MaxRecordedFrame / ConfigRef.FPS };
+		TArray<float> MatTimes { 0.0f, static_cast<float>(MaxRecordedFrame) / ConfigRef.FPS };
 		TrackStreamable->SetMatrixSamples(Mats, MatTimes);
 		
 		if (ensureAlways(TrackStreamable->EndCoding()))
diff --git a/Engine/Plugins/Runtime/GeometryCache/Source/GeometryCache/Private/GeometryCacheSceneProxy.cpp b/Engine/Plugins/Runtime/GeometryCache/Source/GeometryCache/Private/GeometryCacheSceneProxy.cpp
--- a/Engine/Plugins/Runtime/GeometryCache/Source/GeometryCache/Private/GeometryCacheSceneProxy.cpp
+++ b/Engine/Plugins/Runtime/GeometryCache/Source/GeometryCache/Private/GeometryCacheSceneProxy.cpp
@@ -246,10 +246,10 @@ static float OneOver255 = 1.0f / 255.0f;
 inline FPackedNormal InterpolatePackedNormal(const FPackedNormal& A, const FPackedNormal& B, int32 ScaledFactor, int32 OneMinusScaledFactor)
 {
 	FPackedNormal result;
-	result.Vector.X = (A.Vector.X * OneMinusScaledFactor + B.Vector.X * ScaledFactor) * OneOver255;
-	result.Vector.Y = (A.Vector.Y * OneMinusScaledFactor + B.Vector.Y * ScaledFactor) * OneOver255;
-	result.Vector.Z = (A.Vector.Z * OneMinusScaledFactor + B.Vector.Z * ScaledFactor) * OneOver255;
-	result.Vector.W = (A.Vector.W * OneMinusScaledFactor + B.Vector.W * ScaledFactor) * OneOver255;
+	result.Vector.X = static_cast<int8>(static_cast<float>(A.Vector.X * OneMinusScaledFactor + B.Vector.X * ScaledFactor) * OneOver255);
+	result.Vector.Y = static_cast<int8>(static_cast<float>(A.Vector.Y * OneMinusScaledFactor + B.Vector.Y * ScaledFactor) * OneOver255);
+	result.Vector.Z = static_cast<int8>(static_cast<float>(A.Vector.Z * OneMinusScaledFactor + B.Vector.Z * ScaledFactor) * OneOver255);
+	result.Vector.W = static_cast<int8>(static_cast<float>(A.Vector.W * OneMinusScaledFactor + B.Vector.W * ScaledFactor) * OneOver255);
 	return result;
 }
 
@@ -257,10 +257,10 @@ inline FPackedNormal InterpolatePackedNormal(const FPackedNormal& A, const FPack
 inline FColor InterpolatePackedColor(const FColor& A, const FColor& B, int32 ScaledFactor, int32 OneMinusScaledFactor)
 {
 	FColor result;
-	result.R = (A.R * OneMinusScaledFactor + B.R * ScaledFactor) * OneOver255;
-	result.G = (A.G * OneMinusScaledFactor + B.G * ScaledFactor) * OneOver255;
-	result.B = (A.B * OneMinusScaledFactor + B.B * ScaledFactor) * OneOver255;
-	result.A = (A.A * OneMinusScaledFactor + B.A * ScaledFactor) * OneOver255;
+	result.R = static_cast<uint8>(static_cast<float>(A.R * OneMinusScaledFactor + B.R * ScaledFactor) * OneOver255);
+	result.G = static_cast<uint8>(static_cast<float>(A.G * OneMinusScaledFactor + B.G * ScaledFactor) * OneOver255);
+	result.B = static_cast<uint8>(static_cast<float>(A.B * OneMinusScaledFactor + B.B * ScaledFactor) * OneOver255);
+	result.A = static_cast<uint8>(static_cast<float>(A.A * OneMinusScaledFactor + B.A * ScaledFactor) * OneOver255);
 	return result;
 }
 
@@ -538,7 +538,7 @@ void FGeometryCacheSceneProxy::GetDynamicRayTracingInstances(FRayTracingMaterial
 			const int32 MaterialIndex = TrackProxy->Materials.IsValidIndex(BatchInfo.MaterialIndex) ? BatchInfo.MaterialIndex : SegmentIndex;
 			MeshBatch.MaterialRenderProxy = TrackProxy->Materials[MaterialIndex]->GetRenderProxy();
 			MeshBatch.CastRayTracedShadow = IsShadowCast(Context.ReferenceView);
-			MeshBatch.SegmentIndex = SegmentIndex;
+			MeshBatch.SegmentIndex = static_cast<uint8>(SegmentIndex);
 
 			RayTracingInstance.Materials.Add(MeshBatch);
 		}
@@ -583,7 +583,7 @@ uint32 FGeometryCacheSceneProxy::GetMemoryFootprint(void) const
 
 uint32 FGeometryCacheSceneProxy::GetAllocatedSize(void) const
 {
-	return(FPrimitiveSceneProxy::GetAllocatedSize());
+	return static_cast<uint32>(FPrimitiveSceneProxy::GetAllocatedSize());
 }
 
 void FGeometryCacheSceneProxy::UpdateAnimation(FRHICommandListBase& RHICmdList, float NewTime, bool bNewLooping, bool bNewIsPlayingBackwards, float NewPlaybackSpeed, float NewMotionVectorScale)
@@ -819,7 +819,7 @@ void FGeometryCacheSceneProxy::FrameUpdate(FRHICommandListBase& RHICmdList) cons
 
 				Scratch.Prepare(NumVerts, bHasMotionVectors);
 
-				const float OneMinusInterp = 1.0 - InterpolationFactor;
+				const float OneMinusInterp = 1.0f - InterpolationFactor;
 				const int32 InterpFixed = (int32)(InterpolationFactor * 255.0f);
 				const int32 OneMinusInterpFixed = 255 - InterpFixed;
 				const VectorRegister4Float WeightA = VectorSetFloat1( OneMinusInterp );
@@ -1014,7 +1014,7 @@ void FGeometryCacheSceneProxy::FrameUpdate(FRHICommandListBase& RHICmdList) cons
 					// It represents the delta interpolation factor between each sub-frame (due to temporal subsampling)
 					// but we don't want to affect the motion vectors when sampling at multiples of frame so it's clamped to 1
 					float DeltaInterpolationFactor = InterpolationFactor - TrackProxy->PreviousInterpolationFactor;
-					DeltaInterpolationFactor += (TrackProxy->FrameIndex - TrackProxy->PreviousFrameIndex);
+					DeltaInterpolationFactor += static_cast<float>(TrackProxy->FrameIndex - TrackProxy->PreviousFrameIndex);
 					DeltaInterpolationFactor = FMath::Clamp(FMath::Abs(DeltaInterpolationFactor), 0.0f, 1.0f); // the Abs accounts for playing backwards
 					TrackProxy->SubframeInterpolationFactor = FMath::IsNearlyEqual(DeltaInterpolationFactor, 1.0f, KINDA_SMALL_NUMBER) ? 1.0f : DeltaInterpolationFactor;
 
@@ -1208,7 +1208,7 @@ void FGeometryCacheSceneProxy::FrameUpdate(FRHICommandListBase& RHICmdList) cons
 						}
 
 						float DeltaInterpolationFactor = InterpolationFactor - TrackProxy->PreviousInterpolationFactor;
-						DeltaInterpolationFactor += (TrackProxy->FrameIndex - TrackProxy->PreviousFrameIndex);
+						DeltaInterpolationFactor += static_cast<float>(TrackProxy->FrameIndex - TrackProxy->PreviousFrameIndex);
 						DeltaInterpolationFactor = FMath::Clamp(FMath::Abs(DeltaInterpolationFactor), 0.0f, 1.0f);
 						TrackProxy->SubframeInterpolationFactor = FMath::IsNearlyEqual(DeltaInterpolationFactor, 1.0f, KINDA_SMALL_NUMBER) ? 1.0f : DeltaInterpolationFactor;
 
diff --git a/Engine/Plugins/Runtime/GeometryCache/Source/GeometryCache/Private/GeometryCacheStreamingManager.cpp b/Engine/Plugins/Runtime/GeometryCache/Source/GeometryCache/Private/GeometryCacheStreamingManager.cpp
--- a/Engine/Plugins/Runtime/GeometryCache/Source/GeometryCache/Private/GeometryCacheStreamingManager.cpp
+++ b/Engine/Plugins/Runtime/GeometryCache/Source/GeometryCache/Private/GeometryCacheStreamingManager.cpp
@@ -172,7 +172,7 @@ int32 FGeometryCacheStreamingManager::BlockTillAllRequestsFinished(float TimeLim
 		double EndTime = FPlatformTime::Seconds() + TimeLimit;
 		for (auto Iter = StreamingGeometryCaches.CreateIterator(); Iter; ++Iter)
 		{
-			float ThisTimeLimit = EndTime - FPlatformTime::Seconds();
+			float ThisTimeLimit = static_cast<float>(EndTime - FPlatformTime::Seconds());
 			if (ThisTimeLimit < .001f || // one ms is the granularity of the platform event system
 				!Iter.Value()->BlockTillAllRequestsFinished(ThisTimeLimit))
 			{
diff --git a/Engine/Plugins/Runtime/GeometryCache/Source/GeometryCache/Private/GeometryCacheTrackStreamable.cpp b/Engine/Plugins/Runtime/GeometryCache/Source/GeometryCache/Private/GeometryCacheTrackStreamable.cpp
--- a/Engine/Plugins/Runtime/GeometryCache/Source/GeometryCache/Private/GeometryCacheTrackStreamable.cpp
+++ b/Engine/Plugins/Runtime/GeometryCache/Source/GeometryCache/Private/GeometryCacheTrackStreamable.cpp
@@ -34,7 +34,7 @@ void UGeometryCacheTrackStreamable::TriggerSerializationCrash()
 	const int64 AssetSize = 1024 * 1024 * 1024ll; // 1 Gig for the individual asset
 
 	// Add enough bulk data chunks to the requested size
-	const int32 NumChunks = FMath::DivideAndRoundUp(BulkDataSize, ChunkDataSize);
+	const int32 NumChunks = static_cast<int32>(FMath::DivideAndRoundUp(BulkDataSize, ChunkDataSize));
 	for ( int Count=0; Count<NumChunks;Count++)
 	{
 		// Create a new chunk for this coded frame
@@ -52,7 +52,7 @@ void UGeometryCacheTrackStreamable::TriggerSerializationCrash()
 	}
 
 	// Add enough sample info objects to blow the asset up to the requested size
-	int64 NumSamplesToAdd = FMath::DivideAndRoundUp(AssetSize, (int64)sizeof(FGeometryCacheTrackStreamableSampleInfo));
+	const int32 NumSamplesToAdd = static_cast<int32>(FMath::DivideAndRoundUp(AssetSize, (int64)sizeof(FGeometryCacheTrackStreamableSampleInfo)));
 	Track->Samples.AddUninitialized(NumSamplesToAdd);
 
 	FString PackageFileName = FPackageName::LongPackageNameToFilename(PackageName, FPackageName::GetAssetPackageExtension());
@@ -404,7 +404,7 @@ void UGeometryCacheTrackStreamable::GetChunksForTimeRange(float StartTime, float
 
 	for (int32 ChunkID = 0; ChunkID <Chunks.Num(); ChunkID++)
 	{
-		if (Chunks[ChunkID].FirstFrame <= LastFrame && Chunks[ChunkID].LastFrame >= FirstFrame)
+		if (Chunks[ChunkID].FirstFrame <= static_cast<float>(LastFrame) && Chunks[ChunkID].LastFrame >= static_cast<float>(FirstFrame))
 		{
 			OutChunkIndexes.Add(ChunkID);
 		}
diff --git a/Engine/Plugins/Runtime/GeometryCache/Source/GeometryCache/Private/HuffmanBitStream.h b/Engine/Plugins/Runtime/GeometryCache/Source/GeometryCache/Private/HuffmanBitStream.h
--- a/Engine/Plugins/Runtime/GeometryCache/Source/GeometryCache/Private/HuffmanBitStream.h
+++ b/Engine/Plugins/Runtime/GeometryCache/Source/GeometryCache/Private/HuffmanBitStream.h
@@ -252,7 +252,7 @@ public:
 	*/
 	FORCEINLINE int32 GetNumBytes()
 	{
-		return NumBytes;
+		return static_cast<int32>(NumBytes);
 	}
 private:
 	const uint8* Bytes;
diff --git a/Engine/Plugins/Runtime/GeometryCache/Source/GeometryCache/Private/HuffmanTable.cpp b/Engine/Plugins/Runtime/GeometryCache/Source/GeometryCache/Private/HuffmanTable.cpp
--- a/Engine/Plugins/Runtime/GeometryCache/Source/GeometryCache/Private/HuffmanTable.cpp
+++ b/Engine/Plugins/Runtime/GeometryCache/Source/GeometryCache/Private/HuffmanTable.cpp
@@ -158,7 +158,7 @@ static void GenerateHuffmanCodes(uint16* OutCodes, const uint8* CodeLengths, int
 	{
 		int32 CodeLength = CodeLengths[Symbol];
 		check(CodeLength <= HUFFMAN_MAX_CODE_LENGTH);
-		SymbolLists[CodeLength * HUFFMAN_MAX_CODES + SymbolLengthHistogram[CodeLength]++] = Symbol;
+		SymbolLists[CodeLength * HUFFMAN_MAX_CODES + SymbolLengthHistogram[CodeLength]++] = static_cast<int16>(Symbol);
 	}
 
 	// Assign code words to symbols
@@ -171,7 +171,7 @@ static void GenerateHuffmanCodes(uint16* OutCodes, const uint8* CodeLengths, int
 		{
 			int32 Symbol = SymbolLists[CodeLength * HUFFMAN_MAX_CODES + i];
 			check(CodeLengths[Symbol] == CodeLength);
-			OutCodes[Symbol] = ReverseBits(NextCodeWord++, CodeLength);	// Reverse bits to turn prefix codes into postfix codes
+			OutCodes[Symbol] = static_cast<uint16>(ReverseBits(NextCodeWord++, CodeLength));	// Reverse bits to turn prefix codes into postfix codes
 		}
 		NextCodeWord <<= 1;
 	}
@@ -286,7 +286,7 @@ void FHuffmanDecodeTable::Initialize(FHuffmanBitStreamReader &Stream)
 	// Read symbol code lengths
 	for (int32 Symbol = 0; Symbol < NumSymbols; Symbol++)
 	{
-		SymbolLengths[Symbol] = Stream.Read(HUFFMAN_MAX_CODE_LENGTH_BITS);
+		SymbolLengths[Symbol] = static_cast<uint8>(Stream.Read(HUFFMAN_MAX_CODE_LENGTH_BITS));
 	}
 
 	// Generate canonical huffman codes from code lengths
@@ -304,8 +304,8 @@ void FHuffmanDecodeTable::Initialize(FHuffmanBitStreamReader &Stream)
 			uint32 Step = 1u << Length;
 			
 			FTableEntry Entry;
-			Entry.Symbol = Symbol;
-			Entry.Length = Length;
+			Entry.Symbol = static_cast<int8>(Symbol);
+			Entry.Length = static_cast<int8>(Length);
 
 			// Put entry in at every position where the Length last bits are Code.
 			do
diff --git a/Engine/Plugins/Runtime/GeometryCache/Source/GeometryCache/Private/NiagaraRendererGeometryCache.cpp b/Engine/Plugins/Runtime/GeometryCache/Source/GeometryCache/Private/NiagaraRendererGeometryCache.cpp
--- a/Engine/Plugins/Runtime/GeometryCache/Source/GeometryCache/Private/NiagaraRendererGeometryCache.cpp
+++ b/Engine/Plugins/Runtime/GeometryCache/Source/GeometryCache/Private/NiagaraRendererGeometryCache.cpp
@@ -256,7 +256,7 @@ void FNiagaraRendererGeometryCache::PostSystemTick_GameThread(const UNiagaraRend
 	const FVector3f DefaultRotation = ParameterStore.GetParameterValueOrDefault(Properties->RotationBinding.GetParamMapBindableVariable(), FVector3f::ZeroVector);
 	const FVector3f DefaultScale = ParameterStore.GetParameterValueOrDefault(Properties->ScaleBinding.GetParamMapBindableVariable(), FVector3f::OneVector);
 	const float DefaultElapsedTime = ParameterStore.GetParameterValueOrDefault(Properties->ElapsedTimeBinding.GetParamMapBindableVariable(), 0.0f);
-	const float CurrentTime = AttachComponent->GetWorld()->GetRealTimeSeconds();
+	const float CurrentTime = static_cast<float>(AttachComponent->GetWorld()->GetRealTimeSeconds());
 
 	const FNiagaraLWCConverter LwcConverter = SystemInstance->GetLWCConverter(Emitter->IsLocalSpace());
 	const bool bIsRendererEnabled = IsRendererEnabled(InProperties, Emitter);
diff --git a/Engine/Plugins/Runtime/GeometryCache/Source/GeometryCache/Private/NormalCompression.cpp b/Engine/Plugins/Runtime/GeometryCache/Source/GeometryCache/Private/NormalCompression.cpp
--- a/Engine/Plugins/Runtime/GeometryCache/Source/GeometryCache/Private/NormalCompression.cpp
+++ b/Engine/Plugins/Runtime/GeometryCache/Source/GeometryCache/Private/NormalCompression.cpp
@@ -50,9 +50,9 @@ int32 FNormalCoderSmith::GetTotalBinCount() const
 int32 FNormalCoderSmith::CalcNTheta(int32 JIndex, float MaxError)
 {
 	// According to page 3 of Smith et al.
-	float Phi = JIndex * PI / (NPhiValue - 1);
-	float Numerator = FMath::Cos(MaxError) - FMath::Cos(Phi) * FMath::Cos(Phi + PI / (2 * (NPhiValue - 1)));
-	float Denominator = FMath::Sin(Phi) * FMath::Sin(Phi + PI / (2 * (NPhiValue - 1)));
+	float Phi = static_cast<float>(JIndex) * PI / static_cast<float>(NPhiValue - 1);
+	float Numerator = FMath::Cos(MaxError) - FMath::Cos(Phi) * FMath::Cos(Phi + PI / static_cast<float>(2 * (NPhiValue - 1)));
+	float Denominator = FMath::Sin(Phi) * FMath::Sin(Phi + PI / static_cast<float>(2 * (NPhiValue - 1)));
 
 	if (Denominator == 0.0f) // pole
 	{
@@ -78,14 +78,14 @@ FNormalCoderSmith::FEncodedNormal FNormalCoderSmith::Encode(const FVector& Value
 	int32 KIndex = 0;
 
 	// Calculate our bin indices
-	JIndex = FMath::RoundToInt(Spherical.Phi * (NPhiValue - 1) / PI);
+	JIndex = FMath::RoundToInt(Spherical.Phi * static_cast<float>(NPhiValue - 1) / PI);
 	int32 NTheta = GetNtheta(JIndex);
-	KIndex = PositiveModulo(FMath::RoundToInt(Spherical.Theta * NTheta / (2 * PI)), NTheta);
+	KIndex = PositiveModulo(FMath::RoundToInt(Spherical.Theta * static_cast<float>(NTheta) / (2 * PI)), NTheta);
 
 	// Decode those indices again
 	FSphericalCoordinates DecodedDelta;
-	DecodedDelta.Phi = JIndex * PI / (NPhiValue - 1);
-	DecodedDelta.Theta = KIndex * 2 * PI / NTheta;
+	DecodedDelta.Phi = static_cast<float>(JIndex) * PI / static_cast<float>(NPhiValue - 1);
+	DecodedDelta.Theta = static_cast<float>(KIndex) * 2 * PI / static_cast<float>(NTheta);
 
 	// Convert back to absolute positions as our decoder will do
 	FVector DecodedDeltaVector = SphericalToVector(DecodedDelta);
@@ -99,9 +99,9 @@ FVector FNormalCoderSmith::Decode(const FEncodedNormal& Value)
 {
 	// Decode our bin indices to spherical coordinates, according to page 2 of Smith et al.
 	FSphericalCoordinates DecodedDelta;
-	DecodedDelta.Phi = Value.JIndex * PI / (NPhiValue - 1);
+	DecodedDelta.Phi = static_cast<float>(Value.JIndex) * PI / static_cast<float>(NPhiValue - 1);
 	int32 NTheta = GetNtheta(Value.JIndex);
-	DecodedDelta.Theta = Value.KIndex * 2 * PI / NTheta;
+	DecodedDelta.Theta = static_cast<float>(Value.KIndex) * 2 * PI / static_cast<float>(NTheta);
 
 	// Transform our coordinates relative to the previously-seen vector back to absolute values
 	FVector DecodedDeltaVector = SphericalToVector(DecodedDelta);
@@ -116,7 +116,7 @@ FNormalCoderSmith::FSphericalCoordinates FNormalCoderSmith::VectorToSpherical(co
 	FVector NormalizedVector = Vector;
 	NormalizedVector.Normalize();
 	FVector2D UESpherical = NormalizedVector.UnitCartesianToSpherical();
-	return FNormalCoderSmith::FSphericalCoordinates(UESpherical.X, UESpherical.Y);
+	return FNormalCoderSmith::FSphericalCoordinates(static_cast<float>(UESpherical.X), static_cast<float>(UESpherical.Y));
 }
 
 /** Convert from spherical coordinates to Cartesian vectors */
@@ -181,10 +181,10 @@ void FNormalCoderSmith::UpdateRotationMatrix(const FVector& UnmovedVector, FVect
 
 	// Calculate the rotation matrix. Method in the paper did not seem to work. This is a generic rotation between two vectors but 
 	// optimized for the case where one of the vectors is always (0, 0, 1)
-	float SizeTerm = 1.0f / (UnmovedVector.Y * UnmovedVector.Y + UnmovedVector.X * UnmovedVector.X);
-	float X = UnmovedVector.Y * UnmovedVector.Y * SizeTerm * (1.0f - UnmovedVector.Z) + UnmovedVector.Z;
-	float Y = (UnmovedVector.Y * -UnmovedVector.X) * SizeTerm * (1.0f - UnmovedVector.Z);
-	float Z = -UnmovedVector.X;
+	float SizeTerm = static_cast<float>(1.0f / (UnmovedVector.Y * UnmovedVector.Y + UnmovedVector.X * UnmovedVector.X));
+	float X = static_cast<float>(UnmovedVector.Y * UnmovedVector.Y * SizeTerm * (1.0f - UnmovedVector.Z) + UnmovedVector.Z);
+	float Y = static_cast<float>((UnmovedVector.Y * -UnmovedVector.X) * SizeTerm * (1.0f - UnmovedVector.Z));
+	float Z = static_cast<float>(-UnmovedVector.X);
 
 	ColumnZ = UnmovedVector;
 	ColumnX = FVector(X, Y, Z);
@@ -194,13 +194,13 @@ void FNormalCoderSmith::UpdateRotationMatrix(const FVector& UnmovedVector, FVect
 int32 FNormalCoderSmithTest::Denormalize(float Value)
 {
 	// From -1,1 to 0-255
-	return (Value + 1) * 0.5f * 255.0f;
+	return static_cast<int32>((Value + 1) * 0.5f * 255.0f);
 }
 
 FIntVector FNormalCoderSmithTest::Denormalize(const FVector& Value)
 {
 	// From -1,1 to 0-255
-	return FIntVector(Denormalize(Value.X), Denormalize(Value.Y), Denormalize(Value.Z));
+	return FIntVector(Denormalize(static_cast<float>(Value.X)), Denormalize(static_cast<float>(Value.Y)), Denormalize(static_cast<float>(Value.Z)));
 }
 
 void FNormalCoderSmithTest::Test()
@@ -241,14 +241,14 @@ void FNormalCoderSmithTest::Test()
 			MaximumDifferenceSpherical.Y = FMath::Max(MaximumDifferenceSpherical.Y, FMath::Abs(SphericalDifference.Y));
 
 			UE_LOG(LogGeoCaStreamingNormalCompression, Log, TEXT("Input: (%i, %i, %i), Encoded: (%i, %i), Decoded: (%i, %i, %i), Difference: (%i, %i, %i)\n"),
-				Denormalize(Input.X),
-				Denormalize(Input.Y),
-				Denormalize(Input.Z),
+				Denormalize(static_cast<float>(Input.X)),
+				Denormalize(static_cast<float>(Input.Y)),
+				Denormalize(static_cast<float>(Input.Z)),
 				Encoded.JIndex,
 				Encoded.KIndex,
-				Denormalize(Decoded.X),
-				Denormalize(Decoded.Y),
-				Denormalize(Decoded.Z),
+				Denormalize(static_cast<float>(Decoded.X)),
+				Denormalize(static_cast<float>(Decoded.Y)),
+				Denormalize(static_cast<float>(Decoded.Z)),
 				ByteDifference.X, ByteDifference.Y, ByteDifference.Z);
 		}
 	}
diff --git a/Engine/Plugins/Runtime/GeometryCache/Source/GeometryCache/Private/StreamingGeometryCacheData.cpp b/Engine/Plugins/Runtime/GeometryCache/Source/GeometryCache/Private/StreamingGeometryCacheData.cpp
--- a/Engine/Plugins/Runtime/GeometryCache/Source/GeometryCache/Private/StreamingGeometryCacheData.cpp
+++ b/Engine/Plugins/Runtime/GeometryCache/Source/GeometryCache/Private/StreamingGeometryCacheData.cpp
@@ -308,7 +308,7 @@ bool FStreamingGeometryCacheData::BlockTillAllRequestsFinished(float TimeLimit)
 		{
 			if (Iter.Value().IORequest)
 			{
-				float ThisTimeLimit = EndTime - FPlatformTime::Seconds();
+				float ThisTimeLimit = static_cast<float>(EndTime - FPlatformTime::Seconds());
 				if (ThisTimeLimit < .001f || // one ms is the granularity of the platform event system
 					!Iter.Value().IORequest->WaitCompletion(ThisTimeLimit))
 				{
diff --git a/Engine/Plugins/Runtime/GeometryCache/Source/GeometryCache/Public/GeometryCacheSceneProxy.h b/Engine/Plugins/Runtime/GeometryCache/Source/GeometryCache/Public/GeometryCacheSceneProxy.h
--- a/Engine/Plugins/Runtime/GeometryCache/Source/GeometryCache/Public/GeometryCacheSceneProxy.h
+++ b/Engine/Plugins/Runtime/GeometryCache/Source/GeometryCache/Public/GeometryCacheSceneProxy.h
@@ -343,7 +343,7 @@ private:
 		TArray<FColor> InterpolatedColors;
 		TArray<FVector3f> InterpolatedMotionVectors;
 
-		void Prepare(SIZE_T NumVertices, bool bHasMotionVectors)
+		void Prepare(int32 NumVertices, bool bHasMotionVectors)
 		{
 			// Clear entries but keep allocations.
 			InterpolatedPositions.Reset();
diff --git a/Engine/Plugins/Runtime/GeometryCache/Source/GeometryCacheEd/GeometryCacheEd.Build.cs b/Engine/Plugins/Runtime/GeometryCache/Source/GeometryCacheEd/GeometryCacheEd.Build.cs
--- a/Engine/Plugins/Runtime/GeometryCache/Source/GeometryCacheEd/GeometryCacheEd.Build.cs
+++ b/Engine/Plugins/Runtime/GeometryCache/Source/GeometryCacheEd/GeometryCacheEd.Build.cs
@@ -24,5 +24,7 @@ public class GeometryCacheEd : ModuleRules
 				"NiagaraEditor"
 			}
 		);
+
+		UnsafeTypeCastWarningLevel = WarningLevel.Error;
 	}
 }
diff --git a/Engine/Plugins/Runtime/GeometryCache/Source/GeometryCacheEd/Private/GeometryCacheThumbnailScene.cpp b/Engine/Plugins/Runtime/GeometryCache/Source/GeometryCacheEd/Private/GeometryCacheThumbnailScene.cpp
--- a/Engine/Plugins/Runtime/GeometryCache/Source/GeometryCacheEd/Private/GeometryCacheThumbnailScene.cpp
+++ b/Engine/Plugins/Runtime/GeometryCache/Source/GeometryCacheEd/Private/GeometryCacheThumbnailScene.cpp
@@ -48,7 +48,7 @@ void FGeometryCacheThumbnailScene::GetViewMatrixParameters(const float InFOVDegr
 
 	const float HalfFOVRadians = FMath::DegreesToRadians<float>(InFOVDegrees) * 0.5f;
 	// Add extra size to view slightly outside of the sphere to compensate for perspective
-	const float HalfMeshSize = PreviewActor->GetGeometryCacheComponent()->Bounds.SphereRadius * 1.15;
+	const float HalfMeshSize = static_cast<float>(PreviewActor->GetGeometryCacheComponent()->Bounds.SphereRadius * 1.15);
 	const float BoundsZOffset = GetBoundsZOffset(PreviewActor->GetGeometryCacheComponent()->Bounds);
 	const float TargetDistance = HalfMeshSize / FMath::Tan(HalfFOVRadians);
 
diff --git a/Engine/Plugins/Runtime/GeometryCache/Source/GeometryCacheSequencer/Classes/GeometryCacheTrackEditor.cpp b/Engine/Plugins/Runtime/GeometryCache/Source/GeometryCacheSequencer/Classes/GeometryCacheTrackEditor.cpp
--- a/Engine/Plugins/Runtime/GeometryCache/Source/GeometryCacheSequencer/Classes/GeometryCacheTrackEditor.cpp
+++ b/Engine/Plugins/Runtime/GeometryCache/Source/GeometryCacheSequencer/Classes/GeometryCacheTrackEditor.cpp
@@ -102,14 +102,14 @@ int32 FGeometryCacheSection::OnPaintSection(FSequencerSectionPainter& Painter) c
 	// Add lines where the animation starts and ends/loops
 	const float AnimPlayRate = FMath::IsNearlyZero(Section.Params.PlayRate) ? 1.0f : Section.Params.PlayRate;
 	const float Duration = Section.Params.GetSequenceLength();
-	const float SeqLength = Duration - TickResolution.AsSeconds(Section.Params.StartFrameOffset + Section.Params.EndFrameOffset) / AnimPlayRate;
-	const float FirstLoopSeqLength = SeqLength - TickResolution.AsSeconds(Section.Params.FirstLoopStartFrameOffset) / AnimPlayRate;
+	const float SeqLength = static_cast<float>(Duration - TickResolution.AsSeconds(Section.Params.StartFrameOffset + Section.Params.EndFrameOffset) / AnimPlayRate);
+	const float FirstLoopSeqLength = static_cast<float>(SeqLength - TickResolution.AsSeconds(Section.Params.FirstLoopStartFrameOffset) / AnimPlayRate);
 
 	if (!FMath::IsNearlyZero(SeqLength, KINDA_SMALL_NUMBER) && SeqLength > 0)
 	{
-		float MaxOffset = Section.GetRange().Size<FFrameTime>() / TickResolution;
+		float MaxOffset = static_cast<float>(Section.GetRange().Size<FFrameTime>() / TickResolution);
 		float OffsetTime = FirstLoopSeqLength;
-		float StartTime = Section.GetInclusiveStartFrame() / TickResolution;
+		float StartTime = static_cast<float>(Section.GetInclusiveStartFrame() / TickResolution);
 
 		while (OffsetTime < MaxOffset)
 		{
@@ -152,7 +152,7 @@ int32 FGeometryCacheSection::OnPaintSection(FSequencerSectionPainter& Painter) c
 			// Flip the text position if getting near the end of the view range
 			static const float TextOffsetPx = 10.f;
 			bool  bDrawLeft = (Painter.SectionGeometry.Size.X - Time) < (TextSize.X + 22.f) - TextOffsetPx;
-			float TextPosition = bDrawLeft ? Time - TextSize.X - TextOffsetPx : Time + TextOffsetPx;
+			float TextPosition = bDrawLeft ? static_cast<float>(Time - TextSize.X - TextOffsetPx) : Time + TextOffsetPx;
 			//handle mirrored labels
 			const float MajorTickHeight = 9.0f;
 			FVector2D TextOffset(TextPosition, Painter.SectionGeometry.Size.Y - (MajorTickHeight + TextSize.Y));
diff --git a/Engine/Plugins/Runtime/GeometryCache/Source/GeometryCacheSequencer/GeometryCacheSequencer.Build.cs b/Engine/Plugins/Runtime/GeometryCache/Source/GeometryCacheSequencer/GeometryCacheSequencer.Build.cs
--- a/Engine/Plugins/Runtime/GeometryCache/Source/GeometryCacheSequencer/GeometryCacheSequencer.Build.cs
+++ b/Engine/Plugins/Runtime/GeometryCache/Source/GeometryCacheSequencer/GeometryCacheSequencer.Build.cs
@@ -29,6 +29,8 @@ namespace UnrealBuildTool.Rules
                     "GeometryCache"
                 }
             );
-        }
-    }
+
+			UnsafeTypeCastWarningLevel = WarningLevel.Error;
+		}
+	}
 }
diff --git a/Engine/Plugins/Runtime/GeometryCache/Source/GeometryCacheStreamer/GeometryCacheStreamer.Build.cs b/Engine/Plugins/Runtime/GeometryCache/Source/GeometryCacheStreamer/GeometryCacheStreamer.Build.cs
--- a/Engine/Plugins/Runtime/GeometryCache/Source/GeometryCacheStreamer/GeometryCacheStreamer.Build.cs
+++ b/Engine/Plugins/Runtime/GeometryCache/Source/GeometryCacheStreamer/GeometryCacheStreamer.Build.cs
@@ -16,6 +16,8 @@ namespace UnrealBuildTool.Rules
 					"Slate",
 				}
 			);
+
+			UnsafeTypeCastWarningLevel = WarningLevel.Error;
 		}
 	}
 }
diff --git a/Engine/Plugins/Runtime/GeometryCache/Source/GeometryCacheStreamer/Private/GeometryCacheStreamBase.cpp b/Engine/Plugins/Runtime/GeometryCache/Source/GeometryCacheStreamer/Private/GeometryCacheStreamBase.cpp
--- a/Engine/Plugins/Runtime/GeometryCache/Source/GeometryCacheStreamer/Private/GeometryCacheStreamBase.cpp
+++ b/Engine/Plugins/Runtime/GeometryCache/Source/GeometryCacheStreamer/Private/GeometryCacheStreamBase.cpp
@@ -388,7 +388,7 @@ const FGeometryCacheStreamStats& FGeometryCacheStreamBase::GetStreamStats() cons
 	check(IsInGameThread());
 
 	const int32 NumFrames = FramesAvailable.Num();
-	const float Secs = Details.SecondsPerFrame * NumFrames;
+	const float Secs = Details.SecondsPerFrame * static_cast<float>(NumFrames);
 
 	Stats.NumCachedFrames = NumFrames;
 	Stats.CachedDuration = Secs;
@@ -417,8 +417,7 @@ void FGeometryCacheStreamBase::IncrementMemoryStat(const FGeometryCacheMeshData&
 	FResourceSizeEx ResSize;
 	MeshData.GetResourceSizeEx(ResSize);
 
-	float SizeInBytes = ResSize.GetTotalMemoryBytes();
-	MemoryUsed += SizeInBytes / (1024 * 1024);
+	MemoryUsed += static_cast<float>(ResSize.GetTotalMemoryBytes()) / (1024.0f * 1024.0f);
 }
 
 void FGeometryCacheStreamBase::DecrementMemoryStat(const FGeometryCacheMeshData& MeshData)
@@ -426,8 +425,7 @@ void FGeometryCacheStreamBase::DecrementMemoryStat(const FGeometryCacheMeshData&
 	FResourceSizeEx ResSize;
 	MeshData.GetResourceSizeEx(ResSize);
 
-	float SizeInBytes = ResSize.GetTotalMemoryBytes();
-	MemoryUsed -= SizeInBytes / (1024 * 1024);
+	MemoryUsed -= static_cast<float>(ResSize.GetTotalMemoryBytes()) / (1024.0f * 1024.0f);
 }
 
 void FGeometryCacheStreamBase::UpdateCurrentFrameIndex(int32 FrameIndex)
diff --git a/Engine/Plugins/Runtime/GeometryCache/Source/GeometryCacheTracks/GeometryCacheTracks.Build.cs b/Engine/Plugins/Runtime/GeometryCache/Source/GeometryCacheTracks/GeometryCacheTracks.Build.cs
--- a/Engine/Plugins/Runtime/GeometryCache/Source/GeometryCacheTracks/GeometryCacheTracks.Build.cs
+++ b/Engine/Plugins/Runtime/GeometryCache/Source/GeometryCacheTracks/GeometryCacheTracks.Build.cs
@@ -35,6 +35,8 @@ namespace UnrealBuildTool.Rules
                 PublicIncludePathModuleNames.Add("GeometryCacheSequencer");
                 DynamicallyLoadedModuleNames.Add("GeometryCacheSequencer");
             }
-        }
-    }
+
+			UnsafeTypeCastWarningLevel = WarningLevel.Error;
+		}
+	}
 }
diff --git a/Engine/Plugins/Runtime/GeometryCache/Source/GeometryCacheTracks/Private/MovieSceneGeometryCacheSection.cpp b/Engine/Plugins/Runtime/GeometryCache/Source/GeometryCacheTracks/Private/MovieSceneGeometryCacheSection.cpp
--- a/Engine/Plugins/Runtime/GeometryCache/Source/GeometryCacheTracks/Private/MovieSceneGeometryCacheSection.cpp
+++ b/Engine/Plugins/Runtime/GeometryCache/Source/GeometryCacheTracks/Private/MovieSceneGeometryCacheSection.cpp
@@ -167,8 +167,8 @@ void UMovieSceneGeometryCacheSection::GetSnapTimes(TArray<FFrameNumber>& OutSnap
 	const FFrameNumber EndFrame = GetExclusiveEndFrame() - 1; // -1 because we don't need to add the end frame twice
 
 	const float AnimPlayRate = FMath::IsNearlyZero(Params.PlayRate) ? 1.0f : Params.PlayRate;
-	const float SeqLengthSeconds = Params.GetSequenceLength() - FrameRate.AsSeconds(Params.StartFrameOffset + Params.EndFrameOffset) / AnimPlayRate;
-	const float FirstLoopSeqLengthInSeconds = SeqLengthSeconds - FrameRate.AsSeconds(Params.FirstLoopStartFrameOffset) / AnimPlayRate;
+	const float SeqLengthSeconds = static_cast<float>(Params.GetSequenceLength() - FrameRate.AsSeconds(Params.StartFrameOffset + Params.EndFrameOffset) / AnimPlayRate);
+	const float FirstLoopSeqLengthInSeconds = static_cast<float>(SeqLengthSeconds - FrameRate.AsSeconds(Params.FirstLoopStartFrameOffset) / AnimPlayRate);
 
 	const FFrameTime SequenceFrameLength = SeqLengthSeconds * FrameRate;
 	const FFrameTime FirstLoopSequenceFrameLength = FirstLoopSeqLengthInSeconds * FrameRate;
@@ -219,7 +219,7 @@ void UMovieSceneGeometryCacheSection::PostEditChangeProperty(FPropertyChangedEve
 
 		if (!FMath::IsNearlyZero(NewPlayRate))
 		{
-			float CurrentDuration = UE::MovieScene::DiscreteSize(GetRange());
+			float CurrentDuration = static_cast<float>(UE::MovieScene::DiscreteSize(GetRange()));
 			float NewDuration = CurrentDuration * (PreviousPlayRate / NewPlayRate);
 			SetEndFrame(GetInclusiveStartFrame() + FMath::FloorToInt(NewDuration));
 
diff --git a/Engine/Plugins/Runtime/GeometryCache/Source/GeometryCacheTracks/Private/MovieSceneGeometryCacheTemplate.cpp b/Engine/Plugins/Runtime/GeometryCache/Source/GeometryCacheTracks/Private/MovieSceneGeometryCacheTemplate.cpp
--- a/Engine/Plugins/Runtime/GeometryCache/Source/GeometryCacheTracks/Private/MovieSceneGeometryCacheTemplate.cpp
+++ b/Engine/Plugins/Runtime/GeometryCache/Source/GeometryCacheTracks/Private/MovieSceneGeometryCacheTemplate.cpp
@@ -141,16 +141,16 @@ float FMovieSceneGeometryCacheSectionTemplateParameters::MapTimeToAnimation(floa
 	const float SectionPlayRate = PlayRate;
 	const float AnimPlayRate = FMath::IsNearlyZero(SectionPlayRate) ? 1.0f : SectionPlayRate;
 
-	const float FirstLoopSeqLength = SequenceLength - InFrameRate.AsSeconds(FirstLoopStartFrameOffset + StartFrameOffset + EndFrameOffset);
-	const float SeqLength = SequenceLength - InFrameRate.AsSeconds(StartFrameOffset + EndFrameOffset);
+	const float FirstLoopSeqLength = static_cast<float>(SequenceLength - InFrameRate.AsSeconds(FirstLoopStartFrameOffset + StartFrameOffset + EndFrameOffset));
+	const float SeqLength = static_cast<float>(SequenceLength - InFrameRate.AsSeconds(StartFrameOffset + EndFrameOffset));
 
-	float AnimPosition = FFrameTime::FromDecimal((InPosition - SectionStartTime).AsDecimal() * AnimPlayRate) / InFrameRate;
-	AnimPosition += InFrameRate.AsSeconds(FirstLoopStartFrameOffset);
+	float AnimPosition = static_cast<float>(FFrameTime::FromDecimal((InPosition - SectionStartTime).AsDecimal() * AnimPlayRate) / InFrameRate);
+	AnimPosition += static_cast<float>(InFrameRate.AsSeconds(FirstLoopStartFrameOffset));
 	if (SeqLength > 0.f && (bLooping || !FMath::IsNearlyEqual(AnimPosition, SeqLength, 1e-4f)))
 	{
 		AnimPosition = FMath::Fmod(AnimPosition, SeqLength);
 	}
-	AnimPosition += InFrameRate.AsSeconds(StartFrameOffset);
+	AnimPosition += static_cast<float>(InFrameRate.AsSeconds(StartFrameOffset));
 	if (bReverse)
 	{
 		AnimPosition = SequenceLength - AnimPosition;
diff --git a/Engine/Source/Editor/SequencerCore/Public/TimeToPixel.h b/Engine/Source/Editor/SequencerCore/Public/TimeToPixel.h
--- a/Engine/Source/Editor/SequencerCore/Public/TimeToPixel.h
+++ b/Engine/Source/Editor/SequencerCore/Public/TimeToPixel.h
@@ -24,10 +24,10 @@ public:
 		, TickResolution( InTickResolution )
 	{
 		ensureMsgf(WidthPx >= 0, TEXT("Got negative pixel width!"));
-		const float VisibleWidth = static_cast<float>(InLocalViewRange.Size<double>());
+		const double VisibleWidth = InLocalViewRange.Size<double>();
 
 		const float MaxPixelsPerSecond = 1000.f;
-		PixelsPerSecond = VisibleWidth > 0 ? WidthPx / VisibleWidth : MaxPixelsPerSecond;
+		PixelsPerSecond = VisibleWidth > 0 ? static_cast<float>(WidthPx / VisibleWidth) : MaxPixelsPerSecond;
 	}
 
 	/**
@@ -38,7 +38,7 @@ public:
 	 */
 	float SecondsToPixel( double Time ) const
 	{
-		return static_cast<float>(Time - ViewRangeStartSeconds) * PixelsPerSecond;
+		return static_cast<float>((Time - ViewRangeStartSeconds) * PixelsPerSecond);
 	}
 
 	/**
@@ -49,7 +49,7 @@ public:
 	 */
 	float SecondsDeltaToPixel( double TimeDelta ) const
 	{
-		return static_cast<float>(TimeDelta) * PixelsPerSecond;
+		return static_cast<float>(TimeDelta * PixelsPerSecond);
 	}
 
 	/**
@@ -71,7 +71,7 @@ public:
 	 */
 	float FrameToPixel( const FFrameTime& Time ) const
 	{
-		return static_cast<float>(Time / TickResolution - ViewRangeStartSeconds) * PixelsPerSecond;
+		return static_cast<float>((Time / TickResolution - ViewRangeStartSeconds) * PixelsPerSecond);
 	}
 
 	/**
@@ -82,7 +82,7 @@ public:
 	 */
 	float FrameDeltaToPixel( const FFrameTime& TimeDelta) const
 	{
-		return static_cast<float>(TimeDelta / TickResolution) * PixelsPerSecond;
+		return static_cast<float>((TimeDelta / TickResolution) * PixelsPerSecond);
 	}
 
 	/**
diff --git a/Engine/Source/Runtime/Engine/Public/Instances/InstancedPlacementHash.h b/Engine/Source/Runtime/Engine/Public/Instances/InstancedPlacementHash.h
--- a/Engine/Source/Runtime/Engine/Public/Instances/InstancedPlacementHash.h
+++ b/Engine/Source/Runtime/Engine/Public/Instances/InstancedPlacementHash.h
@@ -177,9 +177,9 @@ public:
 		FKey MaxKey = MakeKey(InBox.Max);
 
 		// compute in doubles to avoid overflow issues (queries using WORLD_MAX bounds can produce very large numbers...)
-		double CellsX = (MaxKey.X - MinKey.X + 1);
-		double CellsY = (MaxKey.Y - MinKey.Y + 1);
-		double CellsZ = (MaxKey.Z - MinKey.Z + 1);
+		double CellsX = static_cast<double>((MaxKey.X - MinKey.X + 1));
+		double CellsY = static_cast<double>((MaxKey.Y - MinKey.Y + 1));
+		double CellsZ = static_cast<double>((MaxKey.Z - MinKey.Z + 1));
 		double CellCount = CellsX * CellsY * CellsZ;
 
 		// The idea here is to decide when it is faster to just check every populated cell.
diff --git a/Engine/Source/Runtime/MovieScene/Public/Containers/SparseBitSet.h b/Engine/Source/Runtime/MovieScene/Public/Containers/SparseBitSet.h
--- a/Engine/Source/Runtime/MovieScene/Public/Containers/SparseBitSet.h
+++ b/Engine/Source/Runtime/MovieScene/Public/Containers/SparseBitSet.h
@@ -252,7 +252,7 @@ struct TSparseBitSet
 			uint32 BucketIndex = FMath::CountBits(BucketHash & (HashBit-1));
 
 			const uint32 ThisBitIndex = (BitIndex-BucketSize*Hash);
-			const BucketType ThisBitMask = BucketType(1u) << ThisBitIndex;
+			const BucketType ThisBitMask = static_cast<BucketType>(BucketType(1u) << ThisBitIndex);
 
 			BucketType ThisBucket = Buckets.Get(BucketIndex);
 			if (ThisBucket & ThisBitMask)
diff --git a/Engine/Source/Runtime/MovieScene/Public/Evaluation/MovieScenePlaybackCapabilities.h b/Engine/Source/Runtime/MovieScene/Public/Evaluation/MovieScenePlaybackCapabilities.h
--- a/Engine/Source/Runtime/MovieScene/Public/Evaluation/MovieScenePlaybackCapabilities.h
+++ b/Engine/Source/Runtime/MovieScene/Public/Evaluation/MovieScenePlaybackCapabilities.h
@@ -469,7 +469,7 @@ protected:
 
 		// Get the header for this capability
 		const int32 Index = GetCapabilityIndex(CapabilityBit);
-		const FPlaybackCapabilityHeader& Header = GetHeader(Index);
+		const FPlaybackCapabilityHeader& Header = GetHeader(static_cast<uint8>(Index));
 
 		// Check that we are overwriting the same storage mode
 		using FStorageTraits = TPlaybackCapabilityStorageTraits<StorageType, CapabilityType>;
diff --git a/Engine/Source/Runtime/MovieScene/Public/MovieSceneSection.h b/Engine/Source/Runtime/MovieScene/Public/MovieSceneSection.h
--- a/Engine/Source/Runtime/MovieScene/Public/MovieSceneSection.h
+++ b/Engine/Source/Runtime/MovieScene/Public/MovieSceneSection.h
@@ -804,7 +804,7 @@ inline FFrameNumber GetFirstLoopStartOffsetAtTrimTime(FQualifiedFrameTime TrimTi
 {
 	const float AnimPlayRate = FMath::IsNearlyZero(Params.PlayRate) ? 1.0f : Params.PlayRate;
 	const float AnimPosition = static_cast<float>((TrimTime.Time - StartFrame) / TrimTime.Rate * AnimPlayRate);
-	const float SeqLength = Params.GetSequenceLength() - FrameRate.AsSeconds(Params.StartFrameOffset + Params.EndFrameOffset) / AnimPlayRate;
+	const float SeqLength = static_cast<float>(Params.GetSequenceLength() - FrameRate.AsSeconds(Params.StartFrameOffset + Params.EndFrameOffset) / AnimPlayRate);
 
 	FFrameNumber NewOffset = FrameRate.AsFrameNumber(FMath::Fmod(AnimPosition, SeqLength));
 	NewOffset += Params.FirstLoopStartFrameOffset;
