GeometryCache: Fixed crashes when rendering geometry cache due to parallelization of GDME.

#rnx

diff --git a/Engine/Plugins/Runtime/GeometryCache/Source/GeometryCache/Private/GeometryCacheSceneProxy.cpp b/Engine/Plugins/Runtime/GeometryCache/Source/GeometryCache/Private/GeometryCacheSceneProxy.cpp
--- a/Engine/Plugins/Runtime/GeometryCache/Source/GeometryCache/Private/GeometryCacheSceneProxy.cpp
+++ b/Engine/Plugins/Runtime/GeometryCache/Source/GeometryCache/Private/GeometryCacheSceneProxy.cpp
@@ -149,17 +149,17 @@ FGeometryCacheSceneProxy::FGeometryCacheSceneProxy(UGeometryCacheComponent* Comp
 		Tracks.Add(NewSection);
 	}
 
-	if (IsRayTracingEnabled())
+	// Update at least once after the scene proxy has been constructed
+	// Otherwise it is invisible until animation starts
+	FGeometryCacheSceneProxy* SceneProxy = this;
+	ENQUEUE_RENDER_COMMAND(FGeometryCacheUpdateAnimation)(
+		[SceneProxy](FRHICommandListImmediate& RHICmdList)
 	{
-		// Update at least once after the scene proxy has been constructed
-		// Otherwise it is invisible until animation starts
-		FGeometryCacheSceneProxy* SceneProxy = this;
-		ENQUEUE_RENDER_COMMAND(FGeometryCacheUpdateAnimation)(
-			[SceneProxy](FRHICommandListImmediate& RHICmdList)
-		{
-			SceneProxy->FrameUpdate(RHICmdList);
-		});
+		SceneProxy->FrameUpdate(RHICmdList);
+	});
 
+	if (IsRayTracingEnabled())
+	{
 #if RHI_RAYTRACING
 		{
 			ENQUEUE_RENDER_COMMAND(FGeometryCacheInitRayTracingGeometry)(
@@ -440,13 +440,6 @@ void FGeometryCacheSceneProxy::GetDynamicMeshElements(const TArray<const FSceneV
 
 	if (bVisible)
 	{
-		if (!IsRayTracingEnabled())
-		{
-			// When ray tracing is disabled, update only when visible
-			// This is the old behavior
-			FrameUpdate(RHICmdList);
-		}
-
 		// Iterate over all batches in all tracks and add them to all the relevant views	
 		for (const FGeomCacheTrackProxy* TrackProxy : Tracks)
 		{
@@ -595,11 +588,11 @@ void FGeometryCacheSceneProxy::UpdateAnimation(FRHICommandListBase& RHICmdList,
 	MotionVectorScale = NewMotionVectorScale;
 	UpdatedFrameNum = GFrameNumber + 1;
 
+	// Always update in render thread regardless of visibility, ray tracing or not
+	FrameUpdate(RHICmdList);
+
 	if (IsRayTracingEnabled())
 	{
-		// When ray tracing is enabled, update regardless of visibility
-		FrameUpdate(RHICmdList);
-
 #if RHI_RAYTRACING
 		for (FGeomCacheTrackProxy* Section : Tracks)
 		{
