Fix holdout flag on GeometryCache primitives

Added function to set up dynamic UB parameters using the existing proxy BuildUniformShaderParameters function to avoid divergence in functionality. This pattern should eventually be used in the other dynamic primitives.

diff --git a/Engine/Plugins/Runtime/GeometryCache/Source/GeometryCache/Private/GeometryCacheSceneProxy.cpp b/Engine/Plugins/Runtime/GeometryCache/Source/GeometryCache/Private/GeometryCacheSceneProxy.cpp
--- a/Engine/Plugins/Runtime/GeometryCache/Source/GeometryCache/Private/GeometryCacheSceneProxy.cpp
+++ b/Engine/Plugins/Runtime/GeometryCache/Source/GeometryCache/Private/GeometryCacheSceneProxy.cpp
@@ -19,6 +19,7 @@
 #include "SceneInterface.h"
 #include "SceneManagement.h"
 #include "DataDrivenShaderPlatformInfo.h"
+#include "PrimitiveUniformShaderParametersBuilder.h"
 
 DECLARE_CYCLE_STAT(TEXT("Gather Mesh Elements"), STAT_GeometryCacheSceneProxy_GetMeshElements, STATGROUP_GeometryCache);
 DECLARE_DWORD_COUNTER_STAT(TEXT("Triangle Count"), STAT_GeometryCacheSceneProxy_TriangleCount, STATGROUP_GeometryCache);
@@ -331,16 +332,11 @@ void FGeometryCacheSceneProxy::CreateMeshBatch(
 	Mesh.VertexFactory = &TrackProxy->VertexFactory;
 	Mesh.SegmentIndex = 0;
 
-	const FMatrix& LocalToWorldTransform = TrackProxy->WorldMatrix * GetLocalToWorld();
+	FPrimitiveUniformShaderParametersBuilder Builder;
+	BuildUniformShaderParameters(Builder);
+	Builder.LocalToWorld(TrackProxy->WorldMatrix * GetLocalToWorld());
+	DynamicPrimitiveUniformBuffer.Set(RHICmdList, Builder);
 
-	bool bHasPrecomputedVolumetricLightmap;
-	FMatrix PreviousLocalToWorld;
-	int32 SingleCaptureIndex;
-	bool bOutputVelocity;
-	GetScene().GetPrimitiveUniformShaderParameters_RenderThread(GetPrimitiveSceneInfo(), bHasPrecomputedVolumetricLightmap, PreviousLocalToWorld, SingleCaptureIndex, bOutputVelocity);
-	bOutputVelocity |= AlwaysHasVelocity();
-
-	DynamicPrimitiveUniformBuffer.Set(RHICmdList, LocalToWorldTransform, PreviousLocalToWorld, GetBounds(), GetLocalBounds(), ReceivesDecals(), false, bOutputVelocity);
 	BatchElement.PrimitiveUniformBufferResource = &DynamicPrimitiveUniformBuffer.UniformBuffer;
 
 	const FGeometryCacheMeshData* MeshData = TrackProxy->bNextFrameMeshDataSelected ? TrackProxy->NextFrameMeshData : TrackProxy->MeshData;
diff --git a/Engine/Source/Runtime/Engine/Private/SceneManagement.cpp b/Engine/Source/Runtime/Engine/Private/SceneManagement.cpp
--- a/Engine/Source/Runtime/Engine/Private/SceneManagement.cpp
+++ b/Engine/Source/Runtime/Engine/Private/SceneManagement.cpp
@@ -576,6 +576,14 @@ FDynamicPrimitiveUniformBuffer::~FDynamicPrimitiveUniformBuffer()
 	UniformBuffer.ReleaseResource();
 }
 
+
+void FDynamicPrimitiveUniformBuffer::Set(FRHICommandListBase& RHICmdList, FPrimitiveUniformShaderParametersBuilder& Builder)
+{
+	UniformBuffer.BufferUsage = UniformBuffer_SingleFrame;
+	UniformBuffer.SetContents(RHICmdList, Builder.Build());
+	UniformBuffer.InitResource(RHICmdList);
+}
+
 void FDynamicPrimitiveUniformBuffer::Set(
 	FRHICommandListBase& RHICmdList,
 	const FMatrix& LocalToWorld,
@@ -589,8 +597,7 @@ void FDynamicPrimitiveUniformBuffer::Set(
 	bool bOutputVelocity,
 	const FCustomPrimitiveData* CustomPrimitiveData)
 {
-	UniformBuffer.BufferUsage = UniformBuffer_SingleFrame;
-	UniformBuffer.SetContents(
+	Set(
 		RHICmdList,
 		FPrimitiveUniformShaderParametersBuilder{}
 		.Defaults()
@@ -604,9 +611,7 @@ void FDynamicPrimitiveUniformBuffer::Set(
 			.OutputVelocity(bOutputVelocity)
 			.UseVolumetricLightmap(bHasPrecomputedVolumetricLightmap)
 			.CustomPrimitiveData(CustomPrimitiveData)
-		.Build()
 	);
-	UniformBuffer.InitResource(RHICmdList);
 }
 
 void FDynamicPrimitiveUniformBuffer::Set(
diff --git a/Engine/Source/Runtime/Engine/Public/SceneManagement.h b/Engine/Source/Runtime/Engine/Public/SceneManagement.h
--- a/Engine/Source/Runtime/Engine/Public/SceneManagement.h
+++ b/Engine/Source/Runtime/Engine/Public/SceneManagement.h
@@ -74,6 +74,7 @@ struct FViewMatrices;
 struct FEngineShowFlags;
 class FViewport;
 class FLandscapeRayTracingStateList;
+struct FPrimitiveUniformShaderParametersBuilder;
 
 namespace UE { namespace Color { class FColorSpace; } }
 
@@ -2190,6 +2191,8 @@ public:
 
 	TUniformBuffer<FPrimitiveUniformShaderParameters> UniformBuffer;
 
+	ENGINE_API void Set(FRHICommandListBase& RHICmdList, FPrimitiveUniformShaderParametersBuilder& Builder);
+
 	ENGINE_API void Set(
 		FRHICommandListBase& RHICmdList,
 		const FMatrix& LocalToWorld,
