ISPC: GeometrySceneProxy optimizations

Submitting on behalf of Refik Karic (Intel)

diff --git a/Engine/Plugins/Runtime/GeometryCache/Source/GeometryCache/Private/GeometryCacheSceneProxy.cpp b/Engine/Plugins/Runtime/GeometryCache/Source/GeometryCache/Private/GeometryCacheSceneProxy.cpp
--- a/Engine/Plugins/Runtime/GeometryCache/Source/GeometryCache/Private/GeometryCacheSceneProxy.cpp
+++ b/Engine/Plugins/Runtime/GeometryCache/Source/GeometryCache/Private/GeometryCacheSceneProxy.cpp
@@ -21,6 +21,29 @@
 #include "DataDrivenShaderPlatformInfo.h"
 #include "PrimitiveUniformShaderParametersBuilder.h"
 
+#if INTEL_ISPC
+#include "GeometryCacheSceneProxy.ispc.generated.h"
+#endif
+
+#if INTEL_ISPC
+static_assert(sizeof(uint32) == sizeof(FPackedNormal), "sizeof(uint32) != sizeof(FPackedNormal)");
+static_assert(sizeof(ispc::FVector2f) == sizeof(FVector2f), "sizeof(ispc::FVector2f) != sizeof(FVector2f)");
+static_assert(sizeof(ispc::FVector3f) == sizeof(FVector3f), "sizeof(ispc::FVector3f) != sizeof(FVector3f)");
+
+#if UE_BUILD_SHIPPING
+const bool GGeometryCacheSceneProxyUseIspc = true;
+#else
+bool GGeometryCacheSceneProxyUseIspc = true;
+static FAutoConsoleVariableRef CVarGeometryCacheSceneProxyUseIspc(
+	TEXT("r.GeometryCacheSceneProxy.ISPC"),
+	GGeometryCacheSceneProxyUseIspc,
+	TEXT("When enabled GeometryCacheSceneProxy will use ISPC if appropriate."),
+	ECVF_Default
+);
+#endif
+
+#endif
+
 DECLARE_CYCLE_STAT(TEXT("Gather Mesh Elements"), STAT_GeometryCacheSceneProxy_GetMeshElements, STATGROUP_GeometryCache);
 DECLARE_DWORD_COUNTER_STAT(TEXT("Triangle Count"), STAT_GeometryCacheSceneProxy_TriangleCount, STATGROUP_GeometryCache);
 DECLARE_DWORD_COUNTER_STAT(TEXT("Batch Count"), STAT_GeometryCacheSceneProxy_MeshBatchCount, STATGROUP_GeometryCache);
@@ -28,6 +51,11 @@ DECLARE_CYCLE_STAT(TEXT("Vertex Buffer Update"), STAT_VertexBufferUpdate, STATGR
 DECLARE_CYCLE_STAT(TEXT("Index Buffer Update"), STAT_IndexBufferUpdate, STATGROUP_GeometryCache);
 DECLARE_CYCLE_STAT(TEXT("Buffer Update Task"), STAT_BufferUpdateTask, STATGROUP_GeometryCache);
 DECLARE_CYCLE_STAT(TEXT("InterpolateFrames"), STAT_InterpolateFrames, STATGROUP_GeometryCache);
+DECLARE_CYCLE_STAT(TEXT("InterpolateFramesSIMDPositions"), STAT_InterpolateFramesSIMDPositions, STATGROUP_GeometryCache);
+DECLARE_CYCLE_STAT(TEXT("InterpolateFramesSIMDMotionVectors"), STAT_InterpolateFramesSIMDMotionVectors, STATGROUP_GeometryCache);
+DECLARE_CYCLE_STAT(TEXT("InterpolateFramesSIMDTangents"), STAT_InterpolateFramesSIMDTangents, STATGROUP_GeometryCache);
+DECLARE_CYCLE_STAT(TEXT("InterpolateFramesSIMDUVs"), STAT_InterpolateFramesSIMDUVs, STATGROUP_GeometryCache);
+DECLARE_CYCLE_STAT(TEXT("InterpolateFramesSIMDColors"), STAT_InterpolateFramesSIMDColors, STATGROUP_GeometryCache);
 
 static TAutoConsoleVariable<int32> CVarOffloadUpdate(
 	TEXT("GeometryCache.OffloadUpdate"),
@@ -818,6 +846,7 @@ void FGeometryCacheSceneProxy::FrameUpdate(FRHICommandListBase& RHICmdList) cons
 
 				#define VALIDATE 0
 				{
+
 					check(TrackProxy->MeshData->Positions.Num() >= NumVerts);
 					check(TrackProxy->NextFrameMeshData->Positions.Num() >= NumVerts);
 					check(Scratch.InterpolatedPositions.Num() >= NumVerts);
@@ -825,7 +854,22 @@ void FGeometryCacheSceneProxy::FrameUpdate(FRHICommandListBase& RHICmdList) cons
 					const FVector3f* PositionBPtr = TrackProxy->NextFrameMeshData->Positions.GetData();
 					FVector3f* InterpolatedPositionsPtr = Scratch.InterpolatedPositions.GetData();
 
+					SCOPE_CYCLE_COUNTER(STAT_InterpolateFramesSIMDPositions);
+
 					// Unroll 4 times so we can do 4 wide SIMD
+#if INTEL_ISPC
+					if (GGeometryCacheSceneProxyUseIspc)
+					{
+						ispc::InterpolatePositions(
+							(ispc::FVector3f*)PositionAPtr,
+							(ispc::FVector3f*)PositionBPtr,
+							(ispc::FVector3f*)InterpolatedPositionsPtr,
+							NumVerts,
+							OneMinusInterp,
+							InterpolationFactor);
+					}
+					else
+#endif
 					{
 						const FVector4f* PositionAPtr4 = (const FVector4f*)PositionAPtr;
 						const FVector4f* PositionBPtr4 = (const FVector4f*)PositionBPtr;
@@ -875,55 +919,75 @@ void FGeometryCacheSceneProxy::FrameUpdate(FRHICommandListBase& RHICmdList) cons
 					FPackedNormal* InterpolatedTangentXPtr = Scratch.InterpolatedTangentX.GetData();
 					FPackedNormal* InterpolatedTangentZPtr = Scratch.InterpolatedTangentZ.GetData();
 
-					const uint32 SignMask = 0x80808080u;
-					for (int32 Index = 0; Index < NumVerts; ++Index)
+					SCOPE_CYCLE_COUNTER(STAT_InterpolateFramesSIMDTangents);
+
+#if INTEL_ISPC
+					if (GGeometryCacheSceneProxyUseIspc)
 					{
-						// VectorLoadSignedByte4 on all inputs is significantly more expensive than VectorLoadByte4, so lets just use unsigned.
-						// Interpolating signed values as unsigned is not correct, but if we flip the signs first it is!
-						// Flipping the sign maps the signed range [-128, 127] to the unsigned range [0, 255]
-						// Unsigned value with flip			Signed value
-						// 0								-128
-						// 1								-127
-						// ..								..
-						// 127								-1
-						// 128								0
-						// 129								1
-						// 255								127
-
-						uint32 TangentXA = TangentXAPtr[Index].Vector.Packed ^ SignMask;
-						uint32 TangentXB = TangentXBPtr[Index].Vector.Packed ^ SignMask;
-						VectorRegister4Float InterpolatedTangentX =	VectorMultiplyAdd(	VectorLoadByte4(&TangentXA), WeightA, 
-																VectorMultiplyAdd(	VectorLoadByte4(&TangentXB), WeightB, Half));	// +0.5f so truncation becomes round to nearest.
-						uint32 PackedInterpolatedTangentX;
-						VectorStoreByte4(InterpolatedTangentX, &PackedInterpolatedTangentX);
-						InterpolatedTangentXPtr[Index].Vector.Packed = PackedInterpolatedTangentX ^ SignMask;	// Convert back to signed
-
-						uint32 TangentZA = TangentZAPtr[Index].Vector.Packed ^ SignMask;
-						uint32 TangentZB = TangentZBPtr[Index].Vector.Packed ^ SignMask;
-						VectorRegister4Float InterpolatedTangentZ =	VectorMultiplyAdd(	VectorLoadByte4(&TangentZA), WeightA, 
-																VectorMultiplyAdd(	VectorLoadByte4(&TangentZB), WeightB, Half));	// +0.5f so truncation becomes round to nearest.
-						uint32 PackedInterpolatedTangentZ;
-						VectorStoreByte4(InterpolatedTangentZ, &PackedInterpolatedTangentZ);
-						InterpolatedTangentZPtr[Index].Vector.Packed = PackedInterpolatedTangentZ ^ SignMask;	// Convert back to signed
+						ispc::InterpolateTangents(
+							(uint8*)TangentXAPtr,
+							(uint8*)TangentXBPtr,
+							(uint8*)TangentZAPtr,
+							(uint8*)TangentZBPtr,
+							(uint8*)InterpolatedTangentXPtr,
+							(uint8*)InterpolatedTangentZPtr,
+							NumVerts * sizeof(uint32),
+							OneMinusInterp,
+							InterpolationFactor);
 					}
-					VectorResetFloatRegisters();	//TODO: is this actually needed on any platform?
+					else
+#endif
+					{
+						const uint32 SignMask = 0x80808080u;
+						for (int32 Index = 0; Index < NumVerts; ++Index)
+						{
+							// VectorLoadSignedByte4 on all inputs is significantly more expensive than VectorLoadByte4, so lets just use unsigned.
+							// Interpolating signed values as unsigned is not correct, but if we flip the signs first it is!
+							// Flipping the sign maps the signed range [-128, 127] to the unsigned range [0, 255]
+							// Unsigned value with flip			Signed value
+							// 0								-128
+							// 1								-127
+							// ..								..
+							// 127								-1
+							// 128								0
+							// 129								1
+							// 255								127
+
+							uint32 TangentXA = TangentXAPtr[Index].Vector.Packed ^ SignMask;
+							uint32 TangentXB = TangentXBPtr[Index].Vector.Packed ^ SignMask;
+							VectorRegister4Float InterpolatedTangentX = VectorMultiplyAdd(VectorLoadByte4(&TangentXA), WeightA,
+								VectorMultiplyAdd(VectorLoadByte4(&TangentXB), WeightB, Half));	// +0.5f so truncation becomes round to nearest.
+							uint32 PackedInterpolatedTangentX;
+							VectorStoreByte4(InterpolatedTangentX, &PackedInterpolatedTangentX);
+							InterpolatedTangentXPtr[Index].Vector.Packed = PackedInterpolatedTangentX ^ SignMask;	// Convert back to signed
+
+							uint32 TangentZA = TangentZAPtr[Index].Vector.Packed ^ SignMask;
+							uint32 TangentZB = TangentZBPtr[Index].Vector.Packed ^ SignMask;
+							VectorRegister4Float InterpolatedTangentZ = VectorMultiplyAdd(VectorLoadByte4(&TangentZA), WeightA,
+								VectorMultiplyAdd(VectorLoadByte4(&TangentZB), WeightB, Half));	// +0.5f so truncation becomes round to nearest.
+							uint32 PackedInterpolatedTangentZ;
+							VectorStoreByte4(InterpolatedTangentZ, &PackedInterpolatedTangentZ);
+							InterpolatedTangentZPtr[Index].Vector.Packed = PackedInterpolatedTangentZ ^ SignMask;	// Convert back to signed
+						}
+						VectorResetFloatRegisters();	//TODO: is this actually needed on any platform?
 
 #if VALIDATE
-					for (int32 Index = 0; Index < NumVerts; ++Index)
-					{
-						FPackedNormal ResultX = InterpolatePackedNormal(TangentXAPtr[Index], TangentXBPtr[Index], InterpFixed, OneMinusInterpFixed);
-						FPackedNormal ResultZ = InterpolatePackedNormal(TangentZAPtr[Index], TangentZBPtr[Index], InterpFixed, OneMinusInterpFixed);
-						check(FMath::Abs(InterpolatedTangentXPtr[Index].Vector.X - ResultX.Vector.X) <= 2);
-						check(FMath::Abs(InterpolatedTangentXPtr[Index].Vector.Y - ResultX.Vector.Y) <= 2);
-						check(FMath::Abs(InterpolatedTangentXPtr[Index].Vector.Z - ResultX.Vector.Z) <= 2);
-						check(FMath::Abs(InterpolatedTangentXPtr[Index].Vector.W - ResultX.Vector.W) <= 2);
-
-						check(FMath::Abs(InterpolatedTangentZPtr[Index].Vector.X - ResultZ.Vector.X) <= 2);
-						check(FMath::Abs(InterpolatedTangentZPtr[Index].Vector.Y - ResultZ.Vector.Y) <= 2);
-						check(FMath::Abs(InterpolatedTangentZPtr[Index].Vector.Z - ResultZ.Vector.Z) <= 2);
-						check(FMath::Abs(InterpolatedTangentZPtr[Index].Vector.W - ResultZ.Vector.W) <= 2);
-					}
+						for (int32 Index = 0; Index < NumVerts; ++Index)
+						{
+							FPackedNormal ResultX = InterpolatePackedNormal(TangentXAPtr[Index], TangentXBPtr[Index], InterpFixed, OneMinusInterpFixed);
+							FPackedNormal ResultZ = InterpolatePackedNormal(TangentZAPtr[Index], TangentZBPtr[Index], InterpFixed, OneMinusInterpFixed);
+							check(FMath::Abs(InterpolatedTangentXPtr[Index].Vector.X - ResultX.Vector.X) <= 2);
+							check(FMath::Abs(InterpolatedTangentXPtr[Index].Vector.Y - ResultX.Vector.Y) <= 2);
+							check(FMath::Abs(InterpolatedTangentXPtr[Index].Vector.Z - ResultX.Vector.Z) <= 2);
+							check(FMath::Abs(InterpolatedTangentXPtr[Index].Vector.W - ResultX.Vector.W) <= 2);
+
+							check(FMath::Abs(InterpolatedTangentZPtr[Index].Vector.X - ResultZ.Vector.X) <= 2);
+							check(FMath::Abs(InterpolatedTangentZPtr[Index].Vector.Y - ResultZ.Vector.Y) <= 2);
+							check(FMath::Abs(InterpolatedTangentZPtr[Index].Vector.Z - ResultZ.Vector.Z) <= 2);
+							check(FMath::Abs(InterpolatedTangentZPtr[Index].Vector.W - ResultZ.Vector.W) <= 2);
+						}
 #endif
+					}
 				}
 
 				if (TrackProxy->MeshData->VertexInfo.bHasColor0)
@@ -935,11 +999,27 @@ void FGeometryCacheSceneProxy::FrameUpdate(FRHICommandListBase& RHICmdList) cons
 					const FColor* ColorBPtr = TrackProxy->NextFrameMeshData->Colors.GetData();
 					FColor* InterpolatedColorsPtr = Scratch.InterpolatedColors.GetData();
 
-					for( int32 Index = 0; Index < NumVerts; ++Index )
+					SCOPE_CYCLE_COUNTER(STAT_InterpolateFramesSIMDColors);
+#if INTEL_ISPC
+					if (GGeometryCacheSceneProxyUseIspc)
 					{
-						VectorRegister4Float InterpolatedColor =		VectorMultiplyAdd( VectorLoadByte4( &ColorAPtr[Index] ), WeightA,
-																VectorMultiplyAdd( VectorLoadByte4( &ColorBPtr[Index] ), WeightB, Half ) );	// +0.5f so truncation becomes round to nearest.
-						VectorStoreByte4(InterpolatedColor, &InterpolatedColorsPtr[Index]);
+						ispc::InterpolateColors(
+							(uint8*)ColorAPtr,
+							(uint8*)ColorBPtr,
+							(uint8*)InterpolatedColorsPtr,
+							NumVerts * sizeof(uint32),
+							OneMinusInterp,
+							InterpolationFactor);
+					}
+					else
+#endif
+					{
+						for (int32 Index = 0; Index < NumVerts; ++Index)
+						{
+							VectorRegister4Float InterpolatedColor = VectorMultiplyAdd(VectorLoadByte4(&ColorAPtr[Index]), WeightA,
+								VectorMultiplyAdd(VectorLoadByte4(&ColorBPtr[Index]), WeightB, Half));	// +0.5f so truncation becomes round to nearest.
+							VectorStoreByte4(InterpolatedColor, &InterpolatedColorsPtr[Index]);
+						}
 					}
 #if VALIDATE
 					for(int32 Index = 0; Index < NumVerts; ++Index)
@@ -964,31 +1044,48 @@ void FGeometryCacheSceneProxy::FrameUpdate(FRHICommandListBase& RHICmdList) cons
 					const FVector2f* UVBPtr = TrackProxy->NextFrameMeshData->TextureCoordinates.GetData();
 					FVector2f* InterpolatedUVsPtr = Scratch.InterpolatedUVs.GetData();
 
-					// Unroll 2x so we can use 4 wide ops. OOP will hopefully take care of the rest.
+					SCOPE_CYCLE_COUNTER(STAT_InterpolateFramesSIMDUVs);
+
+#if INTEL_ISPC
+					if (GGeometryCacheSceneProxyUseIspc)
 					{
-						int32 Index = 0;
-						for (; Index + 1 < NumVerts; Index += 2)
+						ispc::InterpolateUVs(
+							(ispc::FVector2f*)UVAPtr,
+							(ispc::FVector2f*)UVBPtr,
+							(ispc::FVector2f*)InterpolatedUVsPtr,
+							NumVerts,
+							OneMinusInterp,
+							InterpolationFactor);
+					}
+					else
+#endif
+					{
+						// Unroll 2x so we can use 4 wide ops. OOP will hopefully take care of the rest.
 						{
-							VectorRegister4Float InterpolatedUVx2 = VectorMultiplyAdd(	VectorLoad(&UVAPtr[Index].X),
-																						WeightA,
-																						VectorMultiply(VectorLoad(&UVBPtr[Index].X), WeightB));
-							VectorStore(InterpolatedUVx2, &(InterpolatedUVsPtr[Index].X));
+							int32 Index = 0;
+							for (; Index + 1 < NumVerts; Index += 2)
+							{
+								VectorRegister4Float InterpolatedUVx2 = VectorMultiplyAdd(VectorLoad(&UVAPtr[Index].X),
+									WeightA,
+									VectorMultiply(VectorLoad(&UVBPtr[Index].X), WeightB));
+								VectorStore(InterpolatedUVx2, &(InterpolatedUVsPtr[Index].X));
+							}
+
+							if (Index < NumVerts)
+							{
+								InterpolatedUVsPtr[Index] = UVAPtr[Index] * OneMinusInterp + UVBPtr[Index] * InterpolationFactor;
+							}
 						}
 
-						if(Index < NumVerts)
+#if VALIDATE
+						for (int32 Index = 0; Index < NumVerts; ++Index)
 						{
-							InterpolatedUVsPtr[Index] = UVAPtr[Index] * OneMinusInterp + UVBPtr[Index] * InterpolationFactor;
+							FVector2D Result = UVAPtr[Index] * OneMinusInterp + UVBPtr[Index] * InterpolationFactor;
+							check(FMath::Abs(InterpolatedUVsPtr[Index].X - Result.X) < 0.01f);
+							check(FMath::Abs(InterpolatedUVsPtr[Index].Y - Result.Y) < 0.01f);
 						}
-					}
-						
-#if VALIDATE
-					for (int32 Index = 0; Index < NumVerts; ++Index)
-					{
-						FVector2D Result = UVAPtr[Index] * OneMinusInterp + UVBPtr[Index] * InterpolationFactor;
-						check(FMath::Abs(InterpolatedUVsPtr[Index].X - Result.X) < 0.01f);
-						check(FMath::Abs(InterpolatedUVsPtr[Index].Y - Result.Y) < 0.01f);
-					}
 #endif
+					}
 				}
 
 				if (bHasMotionVectors)
@@ -1008,6 +1105,23 @@ void FGeometryCacheSceneProxy::FrameUpdate(FRHICommandListBase& RHICmdList) cons
 					DeltaInterpolationFactor = FMath::Clamp(FMath::Abs(DeltaInterpolationFactor), 0.0f, 1.0f); // the Abs accounts for playing backwards
 					TrackProxy->SubframeInterpolationFactor = FMath::IsNearlyEqual(DeltaInterpolationFactor, 1.0f, KINDA_SMALL_NUMBER) ? 1.0f : DeltaInterpolationFactor;
 
+					SCOPE_CYCLE_COUNTER(STAT_InterpolateFramesSIMDMotionVectors);
+
+#if INTEL_ISPC
+					if (GGeometryCacheSceneProxyUseIspc)
+					{
+						ispc::InterpolateMotionVectors(
+							(ispc::FVector3f*)MotionVectorsAPtr,
+							(ispc::FVector3f*)MotionVectorsBPtr,
+							(ispc::FVector3f*)InterpolatedMotionVectorsPtr,
+							NumVerts,
+							OneMinusInterp,
+							InterpolationFactor,
+							MotionVectorScale);
+					}
+					else
+#endif
+
 					// Unroll 4 times so we can do 4 wide SIMD
 					{
 						const VectorRegister4Float ScaledWeightA = VectorSetFloat1(OneMinusInterp * MotionVectorScale);
diff --git a/Engine/Plugins/Runtime/GeometryCache/Source/GeometryCache/Private/GeometryCacheSceneProxy.ispc b/Engine/Plugins/Runtime/GeometryCache/Source/GeometryCache/Private/GeometryCacheSceneProxy.ispc
new file mode 100644
--- /dev/null
+++ b/Engine/Plugins/Runtime/GeometryCache/Source/GeometryCache/Private/GeometryCacheSceneProxy.ispc
@@ -0,0 +1,159 @@
+// Copyright Epic Games, Inc. All Rights Reserved.
+
+#include "Math/Vector.isph"
+
+#define InterpolatePositionsVarying 1
+
+export void InterpolatePositions(
+	uniform FVector3f PositionA[],
+	uniform FVector3f PositionB[],
+	uniform FVector3f Scratch[],
+	uniform int NumVerts,
+	uniform float WeightA,
+	uniform float WeightB)
+{
+#if InterpolatePositionsVarying
+	foreach(i = 0 ... NumVerts)
+	{
+		const uniform int Index = extract(i, 0);
+		FVector3f VecA = VectorLoad(PositionA + Index);
+		FVector3f VecB = VectorLoad(PositionB + Index);
+		FVector3f InterpResult = (VecA * WeightA) + (VecB * WeightB);
+		VectorStore(Scratch + Index, InterpResult);
+	}
+#else
+	for (uniform int i = 0; i < NumVerts; ++i)
+	{
+		uniform FVector3f VecA = PositionA[i];
+		uniform FVector3f VecB = PositionB[i];
+		Scratch[i] = (VecA * WeightA) + (VecB * WeightB);;
+	}
+#endif
+}
+
+#define InterpolateMotionVectorsVarying 1
+
+export void InterpolateMotionVectors(
+	uniform FVector3f PositionA[],
+	uniform FVector3f PositionB[],
+	uniform FVector3f Scratch[],
+	uniform int NumVerts,
+	uniform float WeightA,
+	uniform float WeightB,
+	uniform float MotionVectorScale)
+{
+#if InterpolateMotionVectorsVarying
+	foreach(i = 0 ... NumVerts)
+	{
+		const uniform int Index = extract(i, 0);
+		FVector3f VecA = VectorLoad(PositionA + Index);
+		FVector3f VecB = VectorLoad(PositionB + Index);
+		FVector3f InterpResult = ((VecA * WeightA) + (VecB * WeightB)) * MotionVectorScale;
+		VectorStore(Scratch + Index, InterpResult);
+	}
+#else
+	for (uniform int i = 0; i < NumVerts; ++i)
+	{
+		uniform FVector3f VecA = PositionA[i];
+		uniform FVector3f VecB = PositionB[i];
+		Scratch[i] = ((VecA * WeightA) + (VecB * WeightB)) * MotionVectorScale;
+	}
+#endif
+}
+
+#define InterpolateUVsVarying 1
+
+export void InterpolateUVs(
+	uniform FVector2f PositionA[],
+	uniform FVector2f PositionB[],
+	uniform FVector2f Scratch[],
+	uniform int NumVerts,
+	uniform float WeightA,
+	uniform float WeightB)
+{
+#if InterpolateUVsVarying
+	foreach(i = 0 ... NumVerts)
+	{
+		const uniform int Index = extract(i, 0);
+		FVector2f VecA = VectorLoad(PositionA + Index);
+		FVector2f VecB = VectorLoad(PositionB + Index);
+		FVector2f InterpResult = (VecA * WeightA) + (VecB * WeightB);
+		VectorStore(Scratch + Index, InterpResult);
+	}
+#else
+	for (uniform int i = 0; i < NumVerts; ++i)
+	{
+		uniform FVector2f VecA = PositionA[i];
+		uniform FVector2f VecB = PositionB[i];
+		Scratch[i] = (VecA * WeightA) + (VecB * WeightB);;
+	}
+#endif
+}
+
+export void InterpolateColors(
+	uniform uint8 InColorA[],
+	uniform uint8 InColorB[],
+	uniform uint8 Scratch[],
+	uniform int NumVerts,
+	uniform float WeightA,
+	uniform float WeightB)
+{
+	foreach(i = 0 ... NumVerts)
+	{
+		uint8 ColorA = InColorA[i];
+		uint8 ColorB = InColorB[i];
+
+		float AScale = (float)ColorA;
+		float BScale = (float)ColorB;
+
+		float ScaledChannels = (AScale * WeightA) + ((BScale * WeightB) + 0.5f);
+
+		uint8 ConvColor = (uint8)ScaledChannels;
+
+		Scratch[i] = ConvColor;
+	}
+}
+
+export void InterpolateTangents(
+	uniform uint8 InTangentXA[],
+	uniform uint8 InTangentXB[],
+	uniform uint8 InTangentZA[],
+	uniform uint8 InTangentZB[],
+	uniform uint8 XScratch[],
+	uniform uint8 ZScratch[],
+	uniform int NumVerts,
+	uniform float WeightA,
+	uniform float WeightB)
+{
+	uint8 SignMask = 0x80u;
+
+	foreach(i = 0 ... NumVerts)
+	{
+		uint8 TangentXABytes = InTangentXA[i];
+		uint8 TangentXBBytes = InTangentXB[i];
+		uint8 TangentZABytes = InTangentZA[i];
+		uint8 TangentZBBytes = InTangentZB[i];
+
+		TangentXABytes ^= SignMask;
+		TangentXBBytes ^= SignMask;
+		TangentZABytes ^= SignMask;
+		TangentZBBytes ^= SignMask;
+
+		float XAScale = (float)TangentXABytes;
+		float XBScale = (float)TangentXBBytes;
+		float ZAScale = (float)TangentZABytes;
+		float ZBScale = (float)TangentZBBytes;
+
+		float XScaledChannels = (XAScale * WeightA) + ((XBScale * WeightB) + 0.5f);
+		float ZScaledChannels = (ZAScale * WeightA) + ((ZBScale * WeightB) + 0.5f);
+
+		uint8 XConvChannels = (uint8)XScaledChannels;
+		uint8 ZConvChannels = (uint8)ZScaledChannels;
+
+		XConvChannels ^= SignMask;
+		ZConvChannels ^= SignMask;
+
+		XScratch[i] = XConvChannels;
+		ZScratch[i] = ZConvChannels;
+	}
+}
diff --git a/Engine/Source/Runtime/Core/Public/Math/Vector.isph b/Engine/Source/Runtime/Core/Public/Math/Vector.isph
--- a/Engine/Source/Runtime/Core/Public/Math/Vector.isph
+++ b/Engine/Source/Runtime/Core/Public/Math/Vector.isph
@@ -426,6 +426,19 @@ inline varying FVector2 VectorLoad(const uniform FVector2 *uniform SrcPtr)
 	return Result;
 }
 
+inline void VectorStore(uniform FVector2f* uniform DstPtr, const varying FVector2f& V)
+{
+	if (((1 << TARGET_WIDTH) - 1 ^ lanemask()) == 0)
+	{
+		soa_to_aos2_ispc(V.V[0], V.V[1], (uniform float* uniform)DstPtr);
+	}
+	else
+	{
+#pragma ignore warning(perf)
+		DstPtr[programIndex] = V;
+	}
+}
+
 inline void VectorStore(uniform FVector2 *uniform DstPtr, const varying FVector2 &V)
 {
 	if(((1<<TARGET_WIDTH)-1 ^ lanemask()) == 0)
@@ -765,6 +778,36 @@ inline varying FIntVector VectorGather(const uniform FIntVector *varying SrcPtr)
 	return Result;
 }
 
+inline FVector2f operator+(const FVector2f& A, const FVector2f& B)
+{
+	FVector2f Result;
+
+	Result.V[0] = A.V[0] + B.V[0];
+	Result.V[1] = A.V[1] + B.V[1];
+
+	return Result;
+}
+
+inline FVector2f operator*(const FVector2f& A, const FVector2f& B)
+{
+	FVector2f Result;
+
+	Result.V[0] = A.V[0] * B.V[0];
+	Result.V[1] = A.V[1] * B.V[1];
+
+	return Result;
+}
+
+inline FVector2f operator*(const FVector2f& A, const float F)
+{
+	FVector2f Result;
+
+	Result.V[0] = A.V[0] * F;
+	Result.V[1] = A.V[1] * F;
+
+	return Result;
+}
+
 inline FVector4d operator+(const FVector4d &A, const FVector4d &B)
 {
 	FVector4d Result;
