GeometryCache: Replace AssetActions usage with AssetDefinitions.

diff --git a/Engine/Plugins/Runtime/GeometryCache/Source/GeometryCacheEd/Classes/AssetDefinition_GeometryCache.h b/Engine/Plugins/Runtime/GeometryCache/Source/GeometryCacheEd/Classes/AssetDefinition_GeometryCache.h
new file mode 100644
--- /dev/null
+++ b/Engine/Plugins/Runtime/GeometryCache/Source/GeometryCacheEd/Classes/AssetDefinition_GeometryCache.h
@@ -0,0 +1,23 @@
+// Copyright Epic Games, Inc. All Rights Reserved.
+
+#pragma once
+
+#include "AssetDefinitionDefault.h"
+#include "CoreMinimal.h"
+
+#include "AssetDefinition_GeometryCache.generated.h"
+
+UCLASS()
+class UAssetDefinition_GeometryCache : public UAssetDefinitionDefault
+{
+	GENERATED_BODY()
+
+public:
+	// UAssetDefinition Begin
+	virtual FText GetAssetDisplayName() const override;
+	virtual FLinearColor GetAssetColor() const override;
+	virtual TSoftClassPtr<UObject> GetAssetClass() const override;
+	virtual bool CanImport() const override;
+	virtual TConstArrayView<FAssetCategoryPath> GetAssetCategories() const override;
+	// UAssetDefinition End
+};
diff --git a/Engine/Plugins/Runtime/GeometryCache/Source/GeometryCacheEd/Classes/AssetTypeActions_GeometryCache.h b/Engine/Plugins/Runtime/GeometryCache/Source/GeometryCacheEd/Classes/AssetTypeActions_GeometryCache.h
deleted file mode 100644
diff --git a/Engine/Plugins/Runtime/GeometryCache/Source/GeometryCacheEd/GeometryCacheEd.Build.cs b/Engine/Plugins/Runtime/GeometryCache/Source/GeometryCacheEd/GeometryCacheEd.Build.cs
--- a/Engine/Plugins/Runtime/GeometryCache/Source/GeometryCacheEd/GeometryCacheEd.Build.cs
+++ b/Engine/Plugins/Runtime/GeometryCache/Source/GeometryCacheEd/GeometryCacheEd.Build.cs
@@ -8,6 +8,7 @@ public class GeometryCacheEd : ModuleRules
 	{
         PublicDependencyModuleNames.AddRange(
 			new string[] {
+				"AssetDefinition",
 				"Core",
 				"CoreUObject",
 				"Engine",
diff --git a/Engine/Plugins/Runtime/GeometryCache/Source/GeometryCacheEd/Private/AssetDefinition_GeometryCache.cpp b/Engine/Plugins/Runtime/GeometryCache/Source/GeometryCacheEd/Private/AssetDefinition_GeometryCache.cpp
new file mode 100644
--- /dev/null
+++ b/Engine/Plugins/Runtime/GeometryCache/Source/GeometryCacheEd/Private/AssetDefinition_GeometryCache.cpp
@@ -0,0 +1,36 @@
+// Copyright Epic Games, Inc. All Rights Reserved.
+
+#include "AssetDefinition_GeometryCache.h"
+
+#include "EditorFramework/AssetImportData.h"
+#include "GeometryCache.h"
+
+#define LOCTEXT_NAMESPACE "AssetTypeActions"
+
+FText UAssetDefinition_GeometryCache::GetAssetDisplayName() const
+{
+	return LOCTEXT("AssetTypeActions_GeometryCache", "GeometryCache");
+}
+
+FLinearColor UAssetDefinition_GeometryCache::GetAssetColor() const
+{
+	return FColor(0, 255, 255);
+}
+
+TSoftClassPtr<UObject> UAssetDefinition_GeometryCache::GetAssetClass() const
+{
+	return UGeometryCache::StaticClass();
+}
+
+bool UAssetDefinition_GeometryCache::CanImport() const
+{
+	return true;
+}
+
+TConstArrayView<FAssetCategoryPath> UAssetDefinition_GeometryCache::GetAssetCategories() const
+{
+	static const auto Categories = {EAssetCategoryPaths::Animation};
+	return Categories;
+}
+
+#undef LOCTEXT_NAMESPACE
diff --git a/Engine/Plugins/Runtime/GeometryCache/Source/GeometryCacheEd/Private/AssetTypeActions_GeometryCache.cpp b/Engine/Plugins/Runtime/GeometryCache/Source/GeometryCacheEd/Private/AssetTypeActions_GeometryCache.cpp
deleted file mode 100644
diff --git a/Engine/Plugins/Runtime/GeometryCache/Source/GeometryCacheEd/Private/GeometryCacheEdModule.cpp b/Engine/Plugins/Runtime/GeometryCache/Source/GeometryCacheEd/Private/GeometryCacheEdModule.cpp
--- a/Engine/Plugins/Runtime/GeometryCache/Source/GeometryCacheEd/Private/GeometryCacheEdModule.cpp
+++ b/Engine/Plugins/Runtime/GeometryCache/Source/GeometryCacheEd/Private/GeometryCacheEdModule.cpp
@@ -1,7 +1,7 @@
 // Copyright Epic Games, Inc. All Rights Reserved.
 
 #include "GeometryCacheEdModule.h"
-#include "AssetTypeActions_GeometryCache.h"
+
 #include "GeometryCache.h"
 #include "GeometryCacheAssetBroker.h"
 #include "GeometryCacheComponent.h"
@@ -9,6 +9,7 @@
 #include "GeometryCacheThumbnailRenderer.h"
 #include "NiagaraEditorModule.h"
 #include "NiagaraGeometryCacheRendererProperties.h"
+#include "ThumbnailRendering/ThumbnailManager.h"
 
 IMPLEMENT_MODULE(FGeometryCacheEdModule, GeometryCacheEd)
 
@@ -16,12 +17,6 @@ void FGeometryCacheEdModule::StartupModule()
 {
 	LLM_SCOPE_BYTAG(GeometryCache);
 
-	FAssetToolsModule& AssetToolsModule = FAssetToolsModule::GetModule();
-
-	IAssetTools& AssetTools = AssetToolsModule.Get();
-	AssetAction = new FAssetTypeActions_GeometryCache();
-	AssetTools.RegisterAssetTypeActions(MakeShareable(AssetAction));
-
 	AssetBroker = new FGeometryCacheAssetBroker();
 	FComponentAssetBrokerage::RegisterBroker(MakeShareable(AssetBroker), UGeometryCacheComponent::StaticClass(), true, true);
 
@@ -48,9 +43,6 @@ void FGeometryCacheEdModule::ShutdownModule()
 {
 	if (UObjectInitialized())
 	{
-		FAssetToolsModule& AssetToolsModule = FAssetToolsModule::GetModule();
-		IAssetTools& AssetTools = AssetToolsModule.Get();
-		AssetTools.UnregisterAssetTypeActions(AssetAction->AsShared());
 		FComponentAssetBrokerage::UnregisterBroker(MakeShareable(AssetBroker));
 		UThumbnailManager::Get().UnregisterCustomRenderer(UGeometryCache::StaticClass());
 	}
diff --git a/Engine/Plugins/Runtime/GeometryCache/Source/GeometryCacheEd/Public/GeometryCacheEdModule.h b/Engine/Plugins/Runtime/GeometryCache/Source/GeometryCacheEd/Public/GeometryCacheEdModule.h
--- a/Engine/Plugins/Runtime/GeometryCache/Source/GeometryCacheEd/Public/GeometryCacheEdModule.h
+++ b/Engine/Plugins/Runtime/GeometryCache/Source/GeometryCacheEd/Public/GeometryCacheEdModule.h
@@ -36,7 +36,6 @@ public:
 		return FModuleManager::Get().IsModuleLoaded( "GeometryCacheEd" );
 	}
 private:
-	class FAssetTypeActions_GeometryCache* AssetAction;
 	class FGeometryCacheAssetBroker* AssetBroker;
 };
 
