Fixed reimporting GeometryCache was adding new material slots each time it was reimported.

diff --git a/Engine/Plugins/Runtime/GeometryCache/Source/GeometryCache/Private/GeometryCache.cpp b/Engine/Plugins/Runtime/GeometryCache/Source/GeometryCache/Private/GeometryCache.cpp
--- a/Engine/Plugins/Runtime/GeometryCache/Source/GeometryCache/Private/GeometryCache.cpp
+++ b/Engine/Plugins/Runtime/GeometryCache/Source/GeometryCache/Private/GeometryCache.cpp
@@ -138,6 +138,8 @@ void UGeometryCache::BeginDestroy()
 
 void UGeometryCache::ClearForReimporting()
 {
+	Materials.Reset();
+	MaterialSlotNames.Reset();
 	Tracks.Empty();
 
 	// Flush the resource release commands to the rendering thread to ensure that the edit change doesn't occur while a resource is still allocated
