MLD - Preview for .abc files

- Custom Preview for .abc files.
- Ability to inspect animation with a timeline control
- Editor Viewport for Geometry Cache assets

diff --git a/Engine/Plugins/Runtime/GeometryCache/Source/GeometryCache/Classes/GeometryCacheComponent.h b/Engine/Plugins/Runtime/GeometryCache/Source/GeometryCache/Classes/GeometryCacheComponent.h
--- a/Engine/Plugins/Runtime/GeometryCache/Source/GeometryCache/Classes/GeometryCacheComponent.h
+++ b/Engine/Plugins/Runtime/GeometryCache/Source/GeometryCache/Classes/GeometryCacheComponent.h
@@ -74,7 +74,6 @@ class GEOMETRYCACHE_API UGeometryCacheComponent : public UMeshComponent
 	virtual bool IsMaterialSlotNameValid(FName MaterialSlotName) const override;
 	//~ End UMeshComponent Interface.
 
-
 	/**
 	* OnObjectReimported, Callback function to refresh section data and update scene proxy.
 	*
@@ -164,11 +163,15 @@ public:
 	UFUNCTION(BlueprintCallable, Category = "Components|GeometryCache")
 	void SetStartTimeOffset(const float NewStartTimeOffset);
 
-	/** Set the current animation time for GeometryCache. Includes the influence of elapsed time and SetStartTimeOffset */
+	/** Get the current animation time for GeometryCache. Includes the influence of elapsed time and SetStartTimeOffset */
 	UFUNCTION(BlueprintCallable, Category = "Components|GeometryCache")
 	float GetAnimationTime() const;
 
-	/** Set the current animation time for GeometryCache. Includes the influence of elapsed time and SetStartTimeOffset */
+	/** Get the current elapsed time for GeometryCache. Doesn't include the influence of StartTimeOffset */
+	UFUNCTION(BlueprintCallable, Category = "Components|GeometryCache")
+	float GetElapsedTime() const;
+
+	/** Get the playback direction for GeometryCache. */
 	UFUNCTION(BlueprintCallable, Category = "Components|GeometryCache")
 	float GetPlaybackDirection() const;
 
@@ -184,6 +187,10 @@ public:
 	UFUNCTION(BlueprintCallable, Category = "Components|GeometryCache")
 	int32 GetNumberOfFrames() const;
 
+	/** Get the number of tracks */
+	UFUNCTION(BlueprintCallable, Category = "Components|GeometryCache")
+	int32 GetNumberOfTracks() const;
+
 	/** Override wireframe color? */
 	UFUNCTION(BlueprintCallable, Category = "Components|GeometryCache")
 	void SetOverrideWireframeColor(bool bOverride);
@@ -200,14 +207,16 @@ public:
 	UFUNCTION(BlueprintCallable, Category = "Components|GeometryCache")
 	FLinearColor GetWireframeOverrideColor() const;
 
-
 public:
-	/** Helper to get the frame of the ABC asset at this time*/
+	/** Helper to get the frame of the ABC asset at time provided*/
 	int32 GetFrameAtTime(const float Time) const;
 
 	/** Helper to get the time at this frame */
 	float GetTimeAtFrame(const int32 Frame) const;
 
+	/** Helper to make the animation jump to this time*/
+	void SetCurrentTime(const float Time);
+
 public:
 	/** Functions to override the default TickComponent */
 	void SetManualTick(bool bInManualTick);
@@ -218,6 +227,21 @@ public:
 	UFUNCTION(BlueprintCallable, Category = "Components|GeometryCache")
 	void TickAtThisTime(const float Time, bool bInIsRunning, bool bInBackwards, bool bInIsLooping);
 
+#if WITH_EDITOR
+	// Animation Helpers for Transport Controls
+	void StepForward();
+
+	void ForwardEnd();
+
+	void StepBackward();
+
+	void BackwardEnd();
+
+	void ToggleLooping();
+
+	TArray<FString> GetTrackNames() const;
+#endif
+
 protected:
 	/**
 	* Invalidate both the Matrix and Mesh sample indices
@@ -255,6 +279,17 @@ protected:
 	*/
 	void ClearTrackData();
 
+	/**
+	 * Jumps animation to time specified.
+	 */
+	void JumpAnimationToTime(float Time, bool bInIsRunning, bool bInBackwards, bool bInIsLooping);
+
+	/**
+	 * Helper method to tick animation by one frame in PlayDirection
+	 */
+	void StepAnimationFrame(bool bInBackwards);
+
+protected:
 	UPROPERTY(EditAnywhere, Interp, Category = GeometryCache)
 	bool bRunning;
 
diff --git a/Engine/Plugins/Runtime/GeometryCache/Source/GeometryCache/Private/GeometryCacheComponent.cpp b/Engine/Plugins/Runtime/GeometryCache/Source/GeometryCache/Private/GeometryCacheComponent.cpp
--- a/Engine/Plugins/Runtime/GeometryCache/Source/GeometryCache/Private/GeometryCacheComponent.cpp
+++ b/Engine/Plugins/Runtime/GeometryCache/Source/GeometryCache/Private/GeometryCacheComponent.cpp
@@ -1,14 +1,14 @@
 // Copyright Epic Games, Inc. All Rights Reserved.
 
 #include "GeometryCacheComponent.h"
-#include "GeometryCache.h"
-#include "Logging/MessageLog.h"
 
+#include "GeometryCache.h"
+#include "GeometryCacheHelpers.h"
+#include "GeometryCacheModule.h"
 #include "GeometryCacheSceneProxy.h"
-
-#include "GeometryCacheTrack.h"
 #include "GeometryCacheStreamingManager.h"
-#include "GeometryCacheModule.h"
+#include "GeometryCacheTrack.h"
+#include "Logging/MessageLog.h"
 #include "RenderingThread.h"
 
 #include UE_INLINE_GENERATED_CPP_BY_NAME(GeometryCacheComponent)
@@ -38,7 +38,7 @@ UGeometryCacheComponent::UGeometryCacheComponent(const FObjectInitializer& Objec
 void UGeometryCacheComponent::BeginDestroy()
 {
 	Super::BeginDestroy();
-	ReleaseResources();	
+	ReleaseResources();
 }
 
 void UGeometryCacheComponent::FinishDestroy()
@@ -65,6 +65,58 @@ void UGeometryCacheComponent::ClearTrackData()
 	TrackSections.Empty();
 }
 
+void UGeometryCacheComponent::JumpAnimationToTime(float Time, bool bInIsRunning, bool bInBackwards, bool bInIsLooping)
+{
+	ElapsedTime = Time;
+
+	// Game thread update:
+	// This mainly just updates the matrix and bounding boxes. All render state (meshes) is done on the render thread
+	bool bUpdatedBoundsOrMatrix = false;
+	for (int32 TrackIndex = 0; TrackIndex < NumTracks; ++TrackIndex)
+	{
+		bUpdatedBoundsOrMatrix |= UpdateTrackSection(TrackIndex, Time);
+	}
+
+	if (bUpdatedBoundsOrMatrix)
+	{
+		UpdateLocalBounds();
+		// Mark the transform as dirty, so the bounds are updated and sent to the render thread
+		MarkRenderTransformDirty();
+	}
+
+	// The actual current playback speed. The PlaybackSpeed variable contains the speed it would
+	// play back at if it were running regardless of if we're running or not. The renderer
+	// needs the actual playback speed if not a paused animation with explicit motion vectors
+	// would just keep on blurring as if it were moving even when paused.
+	float ActualPlaybackSpeed = (bInIsRunning) ? PlaybackSpeed : 0.0f;
+
+	// Schedule an update on the render thread
+	if (FGeometryCacheSceneProxy* CastedProxy = static_cast<FGeometryCacheSceneProxy*>(SceneProxy))
+	{
+		FGeometryCacheSceneProxy* InSceneProxy = CastedProxy;
+		float AnimationTime = Time;
+		float InPlaybackSpeed = ActualPlaybackSpeed;
+		float InMotionVectorScale = GetMotionVectorScale();
+		ENQUEUE_RENDER_COMMAND(FGeometryCacheUpdateAnimation)(
+			[InSceneProxy, AnimationTime, bInIsLooping, bInBackwards, InPlaybackSpeed, InMotionVectorScale](FRHICommandList& RHICmdList)
+			{
+				InSceneProxy->UpdateAnimation(RHICmdList, AnimationTime, bInIsLooping, bInBackwards, InPlaybackSpeed, InMotionVectorScale);
+			});
+	}
+}
+
+void UGeometryCacheComponent::StepAnimationFrame(bool bInBackwards)
+{
+	const FFrameRate TargetFramerate = FFrameRate(FMath::RoundToInt32((float)GetNumberOfFrames() / GetDuration()), 1);
+	const float DeltaTime = static_cast<float>(TargetFramerate.AsInterval());
+	const float UnclampedTime = ElapsedTime + (DeltaTime * PlayDirection * GetPlaybackSpeed());
+	const float NewTickedTime = IsLooping()
+		? GeometyCacheHelpers::WrapAnimationTime(UnclampedTime, Duration)
+		: FMath::Clamp(UnclampedTime, 0.0f, Duration);
+
+	JumpAnimationToTime(NewTickedTime, IsPlaying(), bInBackwards, IsLooping());
+}
+
 void UGeometryCacheComponent::SetupTrackData()
 {
 	if (GeometryCache != nullptr)
@@ -93,7 +145,7 @@ void UGeometryCacheComponent::OnUnregister()
 	ClearTrackData();
 }
 
-void UGeometryCacheComponent::TickComponent(float DeltaTime, enum ELevelTick TickType, FActorComponentTickFunction *ThisTickFunction)
+void UGeometryCacheComponent::TickComponent(float DeltaTime, enum ELevelTick TickType, FActorComponentTickFunction* ThisTickFunction)
 {
 	SCOPE_CYCLE_COUNTER(STAT_GeometryCacheComponent_TickComponent);
 	if (GeometryCache && bRunning && !bManualTick)
@@ -101,11 +153,24 @@ void UGeometryCacheComponent::TickComponent(float DeltaTime, enum ELevelTick Tic
 		// Increase total elapsed time since BeginPlay according to PlayDirection and speed
 		ElapsedTime += (DeltaTime * PlayDirection * GetPlaybackSpeed());
 
-		if (ElapsedTime < 0.0f && bLooping)
+		if (bLooping && (ElapsedTime < 0.0f))
 		{
 			ElapsedTime += Duration;
 		}
-
+		else if (!bLooping)
+		{
+			if(ElapsedTime < 0.0f)
+			{
+				ElapsedTime = 0.0f;
+				bRunning = false;
+			}
+			else if (ElapsedTime > Duration)
+			{
+				ElapsedTime = Duration;
+				bRunning = false;
+			}
+		}
+		
 		// Game thread update:
 		// This mainly just updates the matrix and bounding boxes. All render state (meshes) is done on the render thread
 		bool bUpdatedBoundsOrMatrix = false;
@@ -163,42 +228,7 @@ void UGeometryCacheComponent::TickAtThisTime(const float Time, bool bInIsRunning
 {
 	if (bManualTick && GeometryCache && bRunning)
 	{
-		ElapsedTime = Time;
-
-		// Game thread update:
-		// This mainly just updates the matrix and bounding boxes. All render state (meshes) is done on the render thread
-		bool bUpdatedBoundsOrMatrix = false;
-		for (int32 TrackIndex = 0; TrackIndex < NumTracks; ++TrackIndex)
-		{
-			bUpdatedBoundsOrMatrix |= UpdateTrackSection(TrackIndex, Time);
-		}
-
-		if (bUpdatedBoundsOrMatrix)
-		{
-			UpdateLocalBounds();
-			// Mark the transform as dirty, so the bounds are updated and sent to the render thread
-			MarkRenderTransformDirty();
-		}
-
-		// The actual current playback speed. The PlaybackSpeed variable contains the speed it would
-		// play back at if it were running regardless of if we're running or not. The renderer
-		// needs the actual playback speed if not a paused animation with explicit motion vectors
-		// would just keep on blurring as if it were moving even when paused.
-		float ActualPlaybackSpeed = (bInIsRunning) ? PlaybackSpeed : 0.0f;
-
-		// Schedule an update on the render thread
-		if (FGeometryCacheSceneProxy* CastedProxy = static_cast<FGeometryCacheSceneProxy*>(SceneProxy))
-		{
-			FGeometryCacheSceneProxy* InSceneProxy = CastedProxy;
-			float AnimationTime = Time;
-			float InPlaybackSpeed = ActualPlaybackSpeed;
-			float InMotionVectorScale = GetMotionVectorScale();
-			ENQUEUE_RENDER_COMMAND(FGeometryCacheUpdateAnimation)(
-				[InSceneProxy, AnimationTime, bInIsLooping, bInBackwards, InPlaybackSpeed, InMotionVectorScale](FRHICommandList& RHICmdList)
-				{
-					InSceneProxy->UpdateAnimation(RHICmdList, AnimationTime, bInIsLooping, bInBackwards, InPlaybackSpeed, InMotionVectorScale);
-				});
-		}
+		JumpAnimationToTime(Time, bInIsRunning, bInBackwards, bInIsLooping);
 	}
 }
 
@@ -222,7 +252,7 @@ void UGeometryCacheComponent::UpdateLocalBounds()
 	}
 
 	LocalBounds = LocalBox.IsValid ? FBoxSphereBounds(LocalBox) : FBoxSphereBounds(FVector(0, 0, 0), FVector(0, 0, 0), 0); // fall back to reset box sphere bounds
-	
+
 	// This calls CalcBounds above and finally stores the world bounds in the "Bounds" member variable
 	UpdateBounds();
 }
@@ -257,7 +287,7 @@ UMaterialInterface* UGeometryCacheComponent::GetMaterial(int32 MaterialIndex) co
 	// Otherwise get it from the geometry cache
 	else
 	{
-		return GeometryCache ? ( GeometryCache->Materials.IsValidIndex(MaterialIndex) ? GeometryCache->Materials[MaterialIndex] : nullptr) : nullptr;
+		return GeometryCache ? (GeometryCache->Materials.IsValidIndex(MaterialIndex) ? GeometryCache->Materials[MaterialIndex] : nullptr) : nullptr;
 	}
 }
 
@@ -290,7 +320,7 @@ void UGeometryCacheComponent::CreateTrackSection(int32 TrackIndex)
 
 bool UGeometryCacheComponent::UpdateTrackSection(int32 TrackIndex, float Time)
 {
-	checkf(TrackIndex < TrackSections.Num() && GeometryCache != nullptr && TrackIndex < GeometryCache->Tracks.Num(), TEXT("Invalid SectionIndex") );
+	checkf(TrackIndex < TrackSections.Num() && GeometryCache != nullptr && TrackIndex < GeometryCache->Tracks.Num(), TEXT("Invalid SectionIndex"));
 
 	UGeometryCacheTrack* Track = GeometryCache->Tracks[TrackIndex];
 	FTrackRenderData& UpdateSection = TrackSections[TrackIndex];
@@ -385,7 +415,7 @@ void UGeometryCacheComponent::SetExtrapolateFrames(const bool bNewExtrapolating)
 
 bool UGeometryCacheComponent::IsPlayingReversed() const
 {
-	return FMath::IsNearlyEqual( PlayDirection, -1.0f );
+	return FMath::IsNearlyEqual(PlayDirection, -1.0f);
 }
 
 float UGeometryCacheComponent::GetPlaybackSpeed() const
@@ -396,7 +426,7 @@ float UGeometryCacheComponent::GetPlaybackSpeed() const
 void UGeometryCacheComponent::SetPlaybackSpeed(const float NewPlaybackSpeed)
 {
 	// Currently only positive play back speeds are supported
-	PlaybackSpeed = FMath::Clamp( NewPlaybackSpeed, 0.0f, 512.0f );
+	PlaybackSpeed = FMath::Clamp(NewPlaybackSpeed, 0.0f, 512.0f);
 }
 
 float UGeometryCacheComponent::GetMotionVectorScale() const
@@ -444,7 +474,7 @@ bool UGeometryCacheComponent::SetGeometryCache(UGeometryCache* NewGeomCache)
 
 	// Update physics representation right away
 	RecreatePhysicsState();
-	
+
 	// Update this component streaming data.
 	IStreamingManager::Get().NotifyPrimitiveUpdated(this);
 
@@ -478,6 +508,11 @@ float UGeometryCacheComponent::GetAnimationTime() const
 	return ElapsedTime + ClampedStartTimeOffset;
 }
 
+float UGeometryCacheComponent::GetElapsedTime() const
+{
+	return ElapsedTime;
+}
+
 float UGeometryCacheComponent::GetPlaybackDirection() const
 {
 	return PlayDirection;
@@ -523,7 +558,7 @@ void UGeometryCacheComponent::ReleaseResources()
 int32 UGeometryCacheComponent::GetFrameAtTime(const float Time) const
 {
 	const float FrameTime = GetNumberOfFrames() > 1 ? Duration / (float)(GetNumberOfFrames() - 1) : 0.0f;
-	const int32 NormalizedFrame =  FMath::Clamp(FMath::RoundToInt(Time / FrameTime), 0, GetNumberOfFrames() - 1);
+	const int32 NormalizedFrame = FMath::Clamp(FMath::RoundToInt(Time / FrameTime), 0, GetNumberOfFrames() - 1);
 	const int32 StartFrame = GeometryCache != nullptr ? GeometryCache->GetStartFrame() : 0;
 	return StartFrame + NormalizedFrame; //
 }
@@ -536,15 +571,35 @@ float UGeometryCacheComponent::GetTimeAtFrame(const int32 Frame) const
 	return FMath::Clamp(FrameTime * static_cast<float>(Frame - StartFrame), 0.0f, Duration);
 }
 
+void UGeometryCacheComponent::SetCurrentTime(const float Time)
+{
+	if (IsPlaying())
+	{
+		Stop();
+	}
+
+	JumpAnimationToTime(FMath::Clamp(Time, 0.0f, Duration), IsPlaying(), IsPlayingReversed(), IsLooping());
+}
+
 int32 UGeometryCacheComponent::GetNumberOfFrames() const
 {
 	if (GeometryCache)
 	{
-		return GeometryCache->GetEndFrame() - GeometryCache->GetStartFrame()  + 1;
+		return GeometryCache->GetEndFrame() - GeometryCache->GetStartFrame() + 1;
 	}
 	return 0;
 }
 
+int32 UGeometryCacheComponent::GetNumberOfTracks() const
+{
+	if (GeometryCache)
+	{
+		return GeometryCache->Tracks.Num();
+	}
+
+	return 0;
+}
+
 void UGeometryCacheComponent::SetOverrideWireframeColor(bool bOverride)
 {
 	bOverrideWireframeColor = bOverride;
@@ -577,6 +632,65 @@ void UGeometryCacheComponent::PostEditUndo()
 	InvalidateTrackSampleIndices();
 	MarkRenderStateDirty();
 }
+
+void UGeometryCacheComponent::StepForward()
+{
+	if (IsPlaying())
+	{
+		Stop();
+	}
+	PlayDirection = 1.0f;
+	constexpr bool bInBackwards = false;
+	StepAnimationFrame(bInBackwards);
+}
+
+void UGeometryCacheComponent::ForwardEnd()
+{
+	if (IsPlaying())
+	{
+		Stop();
+	}
+	JumpAnimationToTime(Duration, IsPlaying(), IsPlayingReversed(), IsLooping());
+}
+
+void UGeometryCacheComponent::StepBackward()
+{
+	if (IsPlaying())
+	{
+		Stop();
+	}
+	PlayDirection = -1.0f;
+	constexpr bool bInBackwards = true;
+	StepAnimationFrame(bInBackwards);
+}
+
+void UGeometryCacheComponent::BackwardEnd()
+{
+	if (IsPlaying())
+	{
+		Stop();
+	}
+	JumpAnimationToTime(0.0f, IsPlaying(), IsPlayingReversed(), IsLooping());
+}
+
+void UGeometryCacheComponent::ToggleLooping()
+{	
+	SetLooping(!bLooping);
+	if (!bLooping)
+	{
+		JumpAnimationToTime(GeometyCacheHelpers::WrapAnimationTime(ElapsedTime, Duration), IsPlaying(), IsPlayingReversed(), IsLooping());
+	}	
+}
+
+TArray<FString> UGeometryCacheComponent::GetTrackNames() const
+{
+	TArray<FString> Names;
+	for (const TObjectPtr<UGeometryCacheTrack>& Track : GeometryCache->Tracks)
+	{
+		Names.Add(Track->GetName());
+	}
+	return Names;
+}
 #endif
 
 #undef LOCTEXT_NAMESPACE
diff --git a/Engine/Plugins/Runtime/GeometryCache/Source/GeometryCacheEd/Classes/AssetDefinition_GeometryCache.h b/Engine/Plugins/Runtime/GeometryCache/Source/GeometryCacheEd/Classes/AssetDefinition_GeometryCache.h
--- a/Engine/Plugins/Runtime/GeometryCache/Source/GeometryCacheEd/Classes/AssetDefinition_GeometryCache.h
+++ b/Engine/Plugins/Runtime/GeometryCache/Source/GeometryCacheEd/Classes/AssetDefinition_GeometryCache.h
@@ -19,5 +19,6 @@ public:
 	virtual TSoftClassPtr<UObject> GetAssetClass() const override;
 	virtual bool CanImport() const override;
 	virtual TConstArrayView<FAssetCategoryPath> GetAssetCategories() const override;
+	virtual EAssetCommandResult OpenAssets(const FAssetOpenArgs& OpenArgs) const override;
 	// UAssetDefinition End
 };
diff --git a/Engine/Plugins/Runtime/GeometryCache/Source/GeometryCacheEd/GeometryCacheEd.Build.cs b/Engine/Plugins/Runtime/GeometryCache/Source/GeometryCacheEd/GeometryCacheEd.Build.cs
--- a/Engine/Plugins/Runtime/GeometryCache/Source/GeometryCacheEd/GeometryCacheEd.Build.cs
+++ b/Engine/Plugins/Runtime/GeometryCache/Source/GeometryCacheEd/GeometryCacheEd.Build.cs
@@ -18,11 +18,15 @@ public class GeometryCacheEd : ModuleRules
                 "RenderCore",
                 "RHI",
 				"EditorFramework",
-                "UnrealEd",
+				"EditorWidgets",
+				"UnrealEd",
 				"AssetTools",
                 "GeometryCache",
 				"ToolMenus",
-				"NiagaraEditor"
+				"NiagaraEditor",
+				"AdvancedPreviewScene",
+				"SequencerWidgets",
+				"TimeManagement"
 			}
 		);
 
diff --git a/Engine/Plugins/Runtime/GeometryCache/Source/GeometryCacheEd/Private/AssetDefinition_GeometryCache.cpp b/Engine/Plugins/Runtime/GeometryCache/Source/GeometryCacheEd/Private/AssetDefinition_GeometryCache.cpp
--- a/Engine/Plugins/Runtime/GeometryCache/Source/GeometryCacheEd/Private/AssetDefinition_GeometryCache.cpp
+++ b/Engine/Plugins/Runtime/GeometryCache/Source/GeometryCacheEd/Private/AssetDefinition_GeometryCache.cpp
@@ -4,6 +4,7 @@
 
 #include "EditorFramework/AssetImportData.h"
 #include "GeometryCache.h"
+#include "GeometryCacheAssetEditorToolkit.h"
 
 #define LOCTEXT_NAMESPACE "AssetTypeActions"
 
@@ -33,4 +34,17 @@ TConstArrayView<FAssetCategoryPath> UAssetDefinition_GeometryCache::GetAssetCate
 	return Categories;
 }
 
+EAssetCommandResult UAssetDefinition_GeometryCache::OpenAssets(const FAssetOpenArgs& OpenArgs) const
+{
+	for (UGeometryCache* GeometryCacheAsset : OpenArgs.LoadObjects<UGeometryCache>())
+	{
+		if (GeometryCacheAsset != nullptr)
+		{
+			TSharedRef<FGeometryCacheAssetEditorToolkit> NewCustomAssetEditor(new FGeometryCacheAssetEditorToolkit());
+			NewCustomAssetEditor->InitCustomAssetEditor(OpenArgs.GetToolkitMode(), OpenArgs.ToolkitHost, GeometryCacheAsset);
+		}
+	}
+	return EAssetCommandResult::Handled;
+}
+
 #undef LOCTEXT_NAMESPACE
diff --git a/Engine/Plugins/Runtime/GeometryCache/Source/GeometryCacheEd/Private/GeometryCacheAssetEditorToolkit.cpp b/Engine/Plugins/Runtime/GeometryCache/Source/GeometryCacheEd/Private/GeometryCacheAssetEditorToolkit.cpp
new file mode 100644
--- /dev/null
+++ b/Engine/Plugins/Runtime/GeometryCache/Source/GeometryCacheEd/Private/GeometryCacheAssetEditorToolkit.cpp
@@ -0,0 +1,268 @@
+// Copyright Epic Games, Inc. All Rights Reserved.
+
+#include "GeometryCacheAssetEditorToolkit.h"
+
+#include "AdvancedPreviewSceneModule.h"
+#include "CoreMinimal.h"
+#include "GeometryCacheComponent.h"
+#include "GeometryCacheTimelineBindingAsset.h"
+#include "IDetailsView.h"
+#include "Modules/ModuleManager.h"
+#include "PropertyEditorModule.h"
+#include "SEditorViewport.h"
+#include "Styling/AppStyle.h"
+#include "Widgets/Docking/SDockTab.h"
+#include "Widgets/SGeometryCacheEditorViewport.h"
+#include "Widgets/SGeometryCacheTimeline.h"
+
+#define LOCTEXT_NAMESPACE "GeometryCacheCustomAssetEditor"
+
+namespace UE::GeometryCacheAssetEditorToolkit::Private
+{
+	static const FName GeometryCacheEditorAppIdentifier(TEXT("GeometryCacheEditor"));
+	static const FName ToolkitFName(TEXT("GeometryCacheEditor"));
+	static const FName TabId_Viewport(TEXT("GeometryCacheCustomAssetEditor_Render"));
+	static const FName TabId_AssetProperties(TEXT("GeometryCacheCustomAssetEditor_Details"));
+	static const FName TabId_AnimationProperties(TEXT("GeometryCacheCustomAssetEditor_Timeline"));
+	static const FName TabId_PreviewSceneProperties(TEXT("GeometryCacheCustomAssetEditor_PreviewScene"));
+}
+
+void FGeometryCacheAssetEditorToolkit::RegisterTabSpawners(const TSharedRef<class FTabManager>& InTabManager)
+{
+	using namespace UE::GeometryCacheAssetEditorToolkit::Private;
+
+	WorkspaceMenuCategory = InTabManager->AddLocalWorkspaceMenuCategory(LOCTEXT("WorkspaceMenuGeometryCacheEditor", "Geometry Cache Asset Editor"));
+
+	FAssetEditorToolkit::RegisterTabSpawners(InTabManager);
+
+	InTabManager->RegisterTabSpawner(TabId_Viewport, FOnSpawnTab::CreateSP(this, &FGeometryCacheAssetEditorToolkit::SpawnTab_Viewport))
+		.SetDisplayName(LOCTEXT("ViewportTab", "Viewport"))
+		.SetGroup(WorkspaceMenuCategory.ToSharedRef())
+		.SetIcon(FSlateIcon(FAppStyle::GetAppStyleSetName(), "LevelEditor.Tabs.Render"));
+
+	InTabManager->RegisterTabSpawner(TabId_AssetProperties, FOnSpawnTab::CreateSP(this, &FGeometryCacheAssetEditorToolkit::SpawnTab_AssetProperties))
+		.SetDisplayName(LOCTEXT("AssetPropertiesTab", "Details"))
+		.SetGroup(WorkspaceMenuCategory.ToSharedRef())
+		.SetIcon(FSlateIcon(FAppStyle::GetAppStyleSetName(), "LevelEditor.Tabs.Details"));
+
+	InTabManager->RegisterTabSpawner(TabId_AnimationProperties, FOnSpawnTab::CreateSP(this, &FGeometryCacheAssetEditorToolkit::SpawnTab_AnimationProperties))
+		.SetDisplayName(LOCTEXT("AnimationPropertiesTab", "Timeline"))
+		.SetGroup(WorkspaceMenuCategory.ToSharedRef())
+		.SetIcon(FSlateIcon(FAppStyle::GetAppStyleSetName(), "ClassIcon.CurveBase"));
+
+	InTabManager->RegisterTabSpawner(TabId_PreviewSceneProperties, FOnSpawnTab::CreateSP(this, &FGeometryCacheAssetEditorToolkit::SpawnTab_PreviewSceneProperties))
+		.SetDisplayName(LOCTEXT("PreviewScenePropertiesTab", "Preview Settings"))
+		.SetGroup(WorkspaceMenuCategory.ToSharedRef())
+		.SetIcon(FSlateIcon(FAppStyle::GetAppStyleSetName(), "LevelEditor.Tabs.Details"));
+}
+
+void FGeometryCacheAssetEditorToolkit::UnregisterTabSpawners(const TSharedRef<class FTabManager>& InTabManager)
+{
+	using namespace UE::GeometryCacheAssetEditorToolkit::Private;
+
+	FAssetEditorToolkit::UnregisterTabSpawners(InTabManager);
+
+	InTabManager->UnregisterTabSpawner(TabId_Viewport);
+	InTabManager->UnregisterTabSpawner(TabId_AssetProperties);
+	InTabManager->UnregisterTabSpawner(TabId_AnimationProperties);
+	InTabManager->UnregisterTabSpawner(TabId_PreviewSceneProperties);
+}
+
+void FGeometryCacheAssetEditorToolkit::InitCustomAssetEditor(const EToolkitMode::Type Mode, const TSharedPtr<class IToolkitHost>& InitToolkitHost, UGeometryCache* InCustomAsset)
+{
+	using namespace UE::GeometryCacheAssetEditorToolkit::Private;
+
+	GeometryCacheAsset = InCustomAsset;
+
+	ViewportTab = SNew(SGeometryCacheEditorViewport);
+	
+	InitPreviewComponents();
+	
+	FPropertyEditorModule& PropertyEditorModule = FModuleManager::GetModuleChecked<FPropertyEditorModule>("PropertyEditor");
+	FDetailsViewArgs DetailsViewArgs;
+	DetailsViewArgs.NameAreaSettings = FDetailsViewArgs::HideNameArea;
+
+	DetailView_AssetProperties = PropertyEditorModule.CreateDetailView(DetailsViewArgs);
+	
+	const TSharedRef<FTabManager::FLayout> StandaloneDefaultLayout = FTabManager::NewLayout("Standalone_GeometryCacheEditor_Layout")
+		->AddArea
+		(
+			FTabManager::NewPrimaryArea()
+			->SetOrientation(Orient_Vertical)
+			->Split
+			(
+				FTabManager::NewSplitter()
+				->SetOrientation(Orient_Horizontal)
+				->Split
+				(
+					FTabManager::NewSplitter()
+					->SetOrientation(Orient_Vertical)
+					->Split
+					(
+						FTabManager::NewStack()
+						->SetSizeCoefficient(0.9f)
+						->SetHideTabWell(true)
+						->AddTab(TabId_Viewport, ETabState::OpenedTab)
+					)
+					->Split
+					(
+						FTabManager::NewStack()
+						->SetSizeCoefficient(0.2f)
+						->SetHideTabWell(true)
+						->AddTab(TabId_AnimationProperties, ETabState::OpenedTab)
+					)
+				)
+				->Split
+				(
+					FTabManager::NewStack()
+					->SetSizeCoefficient(0.2f)
+					->AddTab(TabId_AssetProperties, ETabState::OpenedTab)
+					->AddTab(TabId_PreviewSceneProperties, ETabState::OpenedTab)
+				)
+			)
+		);
+
+	const bool bCreateDefaultStandaloneMenu = true;
+	const bool bCreateDefaultToolbar = true;
+
+	FAssetEditorToolkit::InitAssetEditor(
+		Mode,
+		InitToolkitHost,
+		GeometryCacheEditorAppIdentifier,
+		StandaloneDefaultLayout,
+		bCreateDefaultStandaloneMenu,
+		bCreateDefaultToolbar,
+		(UObject*)InCustomAsset);
+
+	if (DetailView_AssetProperties.IsValid())
+	{
+		DetailView_AssetProperties->SetObject(Cast<UObject>(GeometryCacheAsset));
+	}
+}
+
+FGeometryCacheAssetEditorToolkit::FGeometryCacheAssetEditorToolkit()
+	: GeometryCacheAsset(nullptr)
+{
+}
+
+FName FGeometryCacheAssetEditorToolkit::GetToolkitFName() const
+{
+	return UE::GeometryCacheAssetEditorToolkit::Private::ToolkitFName;
+}
+
+FText FGeometryCacheAssetEditorToolkit::GetBaseToolkitName() const
+{
+	return LOCTEXT("AppLabel", "Geometry Cache Asset Editor");
+}
+
+FText FGeometryCacheAssetEditorToolkit::GetToolkitName() const
+{
+	return FText::FromString(GeometryCacheAsset->GetName());
+}
+
+FText FGeometryCacheAssetEditorToolkit::GetToolkitToolTipText() const
+{
+	return LOCTEXT("ToolTip", "Geometry Cache Asset Editor");
+}
+
+FString FGeometryCacheAssetEditorToolkit::GetWorldCentricTabPrefix() const
+{
+	return LOCTEXT("WorldCentricTabPrefix", "GeometryCache ").ToString();
+}
+
+FLinearColor FGeometryCacheAssetEditorToolkit::GetWorldCentricTabColorScale() const
+{
+	return FColor(0, 255, 255);
+}
+
+void FGeometryCacheAssetEditorToolkit::OnClose()
+{
+	ViewportTab.Reset();
+	DetailView_AssetProperties.Reset();
+}
+
+void FGeometryCacheAssetEditorToolkit::InitPreviewComponents()
+{
+	if (GeometryCacheAsset == nullptr)
+	{
+		return;
+	}
+
+	if (PreviewGeometryCacheComponent == nullptr)
+	{
+		PreviewGeometryCacheComponent = NewObject<UGeometryCacheComponent>(GetTransientPackage(), NAME_None, RF_Transient);
+		PreviewGeometryCacheComponent->CastShadow = 1;
+		PreviewGeometryCacheComponent->bCastDynamicShadow = 1;
+		PreviewGeometryCacheComponent->SetGeometryCache(GeometryCacheAsset.Get());
+		PreviewGeometryCacheComponent->Activate(true);
+	}
+	
+	BindingAsset = MakeShareable(new FGeometryCacheTimelineBindingAsset(PreviewGeometryCacheComponent));
+}
+
+TSharedRef<SDockTab> FGeometryCacheAssetEditorToolkit::SpawnTab_Viewport(const FSpawnTabArgs& Args)
+{
+	check(Args.GetTabId() == UE::GeometryCacheAssetEditorToolkit::Private::TabId_Viewport);
+	check(PreviewGeometryCacheComponent != nullptr);
+
+	ViewportTab->SetGeometryCacheComponent(PreviewGeometryCacheComponent.Get());
+
+	return SNew(SDockTab)
+		.Label(LOCTEXT("ViewportTab", "Viewport"))
+		.TabColorScale(GetTabColorScale())
+		[
+			ViewportTab.ToSharedRef()
+		];
+}
+
+TSharedRef<SDockTab> FGeometryCacheAssetEditorToolkit::SpawnTab_AssetProperties(const FSpawnTabArgs& Args)
+{
+	check(Args.GetTabId() == UE::GeometryCacheAssetEditorToolkit::Private::TabId_AssetProperties);
+
+	return SNew(SDockTab)
+		.Label(LOCTEXT("AssetPropertiesTab", "Properties"))
+		.TabColorScale(GetTabColorScale())
+		[
+			DetailView_AssetProperties.ToSharedRef()
+		];
+}
+
+TSharedRef<SDockTab> FGeometryCacheAssetEditorToolkit::SpawnTab_AnimationProperties(const FSpawnTabArgs& Args)
+{
+	check(Args.GetTabId() == UE::GeometryCacheAssetEditorToolkit::Private::TabId_AnimationProperties);
+
+	return SNew(SDockTab)
+		.Label(LOCTEXT("AnimationPropertiesTab", "Timeline"))
+		.TabColorScale(GetTabColorScale())
+		[
+			SNew(SGeometryCacheTimeline, BindingAsset.ToSharedRef())
+		];
+}
+
+TSharedRef<SDockTab> FGeometryCacheAssetEditorToolkit::SpawnTab_PreviewSceneProperties(const FSpawnTabArgs& Args)
+{
+	check(Args.GetTabId() == UE::GeometryCacheAssetEditorToolkit::Private::TabId_PreviewSceneProperties);
+
+	FAdvancedPreviewSceneModule& AdvancedPreviewSceneModule = FModuleManager::LoadModuleChecked<FAdvancedPreviewSceneModule>(TEXT("AdvancedPreviewScene"));
+
+	TSharedRef<SWidget> PreviewSceneSettingsWidget = SNullWidget::NullWidget;
+	TSharedPtr<class FAdvancedPreviewScene> PreviewScene = ViewportTab->GetAdvancedPreviewScene();
+	if (PreviewScene.IsValid())
+	{
+		PreviewSceneSettingsWidget = AdvancedPreviewSceneModule.CreateAdvancedPreviewSceneSettingsWidget(PreviewScene.ToSharedRef());
+	}
+
+	TSharedRef<SDockTab> SpawnedTab =
+		SNew(SDockTab)
+		.Label(LOCTEXT("PreviewSceneSettingsTab", "Preview Scene Settings"))
+		[
+			SNew(SBox)
+				[
+					PreviewSceneSettingsWidget
+				]
+		];
+
+	return SpawnedTab;
+}
+
+#undef LOCTEXT_NAMESPACE
\ No newline at end of file
diff --git a/Engine/Plugins/Runtime/GeometryCache/Source/GeometryCacheEd/Private/GeometryCacheAssetEditorToolkit.h b/Engine/Plugins/Runtime/GeometryCache/Source/GeometryCacheEd/Private/GeometryCacheAssetEditorToolkit.h
new file mode 100644
--- /dev/null
+++ b/Engine/Plugins/Runtime/GeometryCache/Source/GeometryCacheEd/Private/GeometryCacheAssetEditorToolkit.h
@@ -0,0 +1,61 @@
+// Copyright Epic Games, Inc. All Rights Reserved.
+
+#pragma once
+
+#include "CoreMinimal.h"
+#include "Toolkits/AssetEditorToolkit.h"
+#include "GeometryCache.h"
+
+class IDetailsView;
+class SDockableTab;
+class UGeometryCacheComponent;
+class SGeometryCacheTimeline;
+class SGeometryCacheEditorViewport;
+class FGeometryCacheTimelineBindingAsset;
+
+class FGeometryCacheAssetEditorToolkit : public FAssetEditorToolkit
+{
+public:
+	virtual void RegisterTabSpawners(const TSharedRef<class FTabManager>& InTabManager) override;
+	virtual void UnregisterTabSpawners(const TSharedRef<class FTabManager>& InTabManager) override;
+
+	/**
+	 * Edits the specified asset object
+	 *
+	 * @param	Mode					Asset editing mode for this editor (standalone or world-centric)
+	 * @param	InitToolkitHost			When Mode is WorldCentric, this is the level editor instance to spawn this editor within
+	 * @param	InCustomAsset			The Custom Asset to Edit
+	 */
+	void InitCustomAssetEditor(const EToolkitMode::Type Mode, const TSharedPtr<class IToolkitHost>& InitToolkitHost, UGeometryCache * InCustomAsset);
+
+	FGeometryCacheAssetEditorToolkit();
+	virtual ~FGeometryCacheAssetEditorToolkit() = default;
+
+	/** Begin IToolkit interface */
+	virtual FName GetToolkitFName() const override;
+	virtual FText GetBaseToolkitName() const override;
+	virtual FText GetToolkitName() const override;
+	virtual FText GetToolkitToolTipText() const override;
+	virtual FString GetWorldCentricTabPrefix() const override;
+	virtual FLinearColor GetWorldCentricTabColorScale() const override;
+	virtual bool IsPrimaryEditor() const override { return true; }
+	virtual void OnClose() override;
+	/** End IToolkit interface */
+
+private:
+	void InitPreviewComponents();
+
+private:
+	TSharedRef<SDockTab> SpawnTab_Viewport(const FSpawnTabArgs& Args);
+	TSharedRef<SDockTab> SpawnTab_AssetProperties(const FSpawnTabArgs& Args);
+	TSharedRef<SDockTab> SpawnTab_AnimationProperties(const FSpawnTabArgs& Args);
+	TSharedRef<SDockTab> SpawnTab_PreviewSceneProperties(const FSpawnTabArgs& Args);
+
+private:
+	TWeakObjectPtr<UGeometryCache> GeometryCacheAsset;
+	TWeakObjectPtr<UGeometryCacheComponent> PreviewGeometryCacheComponent;
+
+	TSharedPtr<SGeometryCacheEditorViewport> ViewportTab;
+	TSharedPtr<IDetailsView> DetailView_AssetProperties;
+	TSharedPtr<FGeometryCacheTimelineBindingAsset> BindingAsset;
+};
\ No newline at end of file
diff --git a/Engine/Plugins/Runtime/GeometryCache/Source/GeometryCacheEd/Private/GeometryCacheTimeSliderController.cpp b/Engine/Plugins/Runtime/GeometryCache/Source/GeometryCacheEd/Private/GeometryCacheTimeSliderController.cpp
new file mode 100644
--- /dev/null
+++ b/Engine/Plugins/Runtime/GeometryCache/Source/GeometryCacheEd/Private/GeometryCacheTimeSliderController.cpp
@@ -0,0 +1,990 @@
+// Copyright Epic Games, Inc. All Rights Reserved.
+
+#include "GeometryCacheTimeSliderController.h"
+
+#include "CommonFrameRates.h"
+#include "Fonts/FontMeasure.h"
+#include "Framework/Application/SlateApplication.h"
+#include "MovieSceneTimeHelpers.h"
+#include "Styling/AppStyle.h"
+#include "Widgets/SGeometryCacheTimeline.h"
+
+namespace ScrubConstants
+{
+	/**The fraction of the current view range to scroll per unit delta  */
+	constexpr float ScrollPanFraction = 0.1f;
+
+}
+
+struct FGeometryCacheTimeSlideController::FScrubRangeToScreen
+{
+	double ViewStart;
+	float  PixelsPerInput;
+
+	FScrubRangeToScreen(const TRange<double>& InViewInput, const FVector2D& InWidgetSize)
+	{
+		const double ViewInputRange = InViewInput.Size<double>();
+
+		ViewStart = InViewInput.GetLowerBoundValue();
+		PixelsPerInput = ViewInputRange > 0 ? static_cast<float>(InWidgetSize.X / ViewInputRange) : 0;
+	}
+
+	/** Local Widget Space -> Curve Input domain. */
+	double LocalXToInput(float ScreenX) const
+	{
+		return PixelsPerInput > 0 ? (ScreenX / PixelsPerInput) + ViewStart : ViewStart;
+	}
+
+	/** Curve Input domain -> local Widget Space */
+	float InputToLocalX(double Input) const
+	{
+		return static_cast<float>((Input - ViewStart) * PixelsPerInput);
+	}
+};
+
+struct FGeometryCacheTimeSlideController::FDrawTickArgs
+{
+	/** Geometry of the area */
+	FGeometry AllottedGeometry;
+	/** Culling rect of the area */
+	FSlateRect CullingRect;
+	/** Color of each tick */
+	FLinearColor TickColor;
+	/** Offset in Y where to start the tick */
+	float TickOffset;
+	/** Height in of major ticks */
+	float MajorTickHeight;
+	/** Start layer for elements */
+	int32 StartLayer;
+	/** Draw effects to apply */
+	ESlateDrawEffect DrawEffects;
+	/** Whether or not to only draw major ticks */
+	bool bOnlyDrawMajorTicks;
+	/** Whether or not to mirror labels */
+	bool bMirrorLabels;
+
+};
+
+FGeometryCacheTimeSlideController::FGeometryCacheTimeSlideController(const FTimeSliderArgs& InArgs, TWeakPtr<SGeometryCacheTimeline> GeometryCacheTimeline)
+	: TimeSliderArgs(InArgs)
+	, WeakTimeline(GeometryCacheTimeline)
+{
+	ScrubFillBrush = FAppStyle::GetBrush(TEXT("Sequencer.Timeline.ScrubFill"));
+	ScrubHandleUpBrush = FAppStyle::GetBrush(TEXT("Sequencer.Timeline.VanillaScrubHandleUp"));
+	ScrubHandleDownBrush = FAppStyle::GetBrush(TEXT("Sequencer.Timeline.VanillaScrubHandleDown"));
+	EditableTimeBrush = FAppStyle::GetBrush(TEXT("AnimTimeline.SectionMarker"));
+}
+
+FGeometryCacheTimeSlideController::~FGeometryCacheTimeSlideController()
+{
+}
+
+int32 FGeometryCacheTimeSlideController::OnPaintTimeSlider(bool bMirrorLabels, const FGeometry& AllottedGeometry, const FSlateRect& MyCullingRect, FSlateWindowElementList& OutDrawElements, int32 LayerId, const FWidgetStyle& InWidgetStyle, bool bParentEnabled) const
+{
+	const bool bEnabled = bParentEnabled;
+	const ESlateDrawEffect DrawEffects = bEnabled ? ESlateDrawEffect::None : ESlateDrawEffect::DisabledEffect;
+
+	const TRange<double> LocalViewRange = TimeSliderArgs.ViewRange.Get();
+	const float    LocalViewRangeMin = static_cast<float>(LocalViewRange.GetLowerBoundValue());
+	const float    LocalViewRangeMax = static_cast<float>(LocalViewRange.GetUpperBoundValue());
+	const float    LocalSequenceLength = LocalViewRangeMax - LocalViewRangeMin;
+	const TRange<FFrameNumber> LocalPlaybackRange = TimeSliderArgs.PlaybackRange.Get();
+
+	if (LocalSequenceLength > 0)
+	{
+		FScrubRangeToScreen RangeToScreen(LocalViewRange, AllottedGeometry.Size);
+
+		// draw tick marks
+		constexpr float MajorTickHeight = 9.0f;
+
+		FDrawTickArgs Args;
+		{
+			Args.AllottedGeometry = AllottedGeometry;
+			Args.bMirrorLabels = bMirrorLabels;
+			Args.bOnlyDrawMajorTicks = false;
+			Args.TickColor = FLinearColor::White;
+			Args.CullingRect = MyCullingRect;
+			Args.DrawEffects = DrawEffects;
+			Args.StartLayer = LayerId;
+			Args.TickOffset = bMirrorLabels ? 0.0f : FMath::Abs(static_cast<float>(AllottedGeometry.Size.Y) - MajorTickHeight);
+			Args.MajorTickHeight = MajorTickHeight;
+		}
+
+		DrawTicks(OutDrawElements, LocalViewRange, RangeToScreen, Args);
+
+		// draw playback & selection range
+		FPaintPlaybackRangeArgs PlaybackRangeArgs(
+			bMirrorLabels ? FAppStyle::GetBrush("Sequencer.Timeline.PlayRange_Bottom_L") : FAppStyle::GetBrush("Sequencer.Timeline.PlayRange_Top_L"),
+			bMirrorLabels ? FAppStyle::GetBrush("Sequencer.Timeline.PlayRange_Bottom_R") : FAppStyle::GetBrush("Sequencer.Timeline.PlayRange_Top_R"),
+			6.f
+		);
+
+		LayerId = DrawPlaybackRange(AllottedGeometry, MyCullingRect, OutDrawElements, LayerId, RangeToScreen, PlaybackRangeArgs);
+
+		PlaybackRangeArgs.SolidFillOpacity = 0.05f;
+		LayerId = DrawSelectionRange(AllottedGeometry, MyCullingRect, OutDrawElements, LayerId, RangeToScreen, PlaybackRangeArgs);
+
+		// Draw the scrub handle
+		const float      HandleStart = RangeToScreen.InputToLocalX(TimeSliderArgs.ScrubPosition.Get().AsDecimal() / GetTickResolution().AsDecimal()) - 7.0f;
+		const float      HandleEnd = HandleStart + 13.0f;
+
+		const int32 ArrowLayer = LayerId + 2;
+		FPaintGeometry MyGeometry = AllottedGeometry.ToPaintGeometry(FVector2f(HandleEnd - HandleStart, AllottedGeometry.Size.Y), FSlateLayoutTransform(FVector2f(HandleStart, 0.f)));
+		FLinearColor ScrubColor = InWidgetStyle.GetColorAndOpacityTint();
+		{
+			ScrubColor.A = ScrubColor.A * 0.75f;
+			ScrubColor.B *= 0.1f;
+			ScrubColor.G *= 0.2f;
+		}
+
+		const FSlateBrush* Brush = (bMirrorLabels ? ScrubHandleUpBrush : ScrubHandleDownBrush);
+
+		FSlateDrawElement::MakeBox(
+			OutDrawElements,
+			ArrowLayer,
+			MyGeometry,
+			Brush,
+			DrawEffects,
+			ScrubColor
+		);
+
+		{
+			// Draw the current time next to the scrub handle
+			FString FrameString = TimeSliderArgs.NumericTypeInterface->ToString(TimeSliderArgs.ScrubPosition.Get().GetFrame().Value);
+			FSlateFontInfo SmallLayoutFont = FCoreStyle::GetDefaultFontStyle("Regular", 10);
+
+			const TSharedRef< FSlateFontMeasure > FontMeasureService = FSlateApplication::Get().GetRenderer()->GetFontMeasureService();
+			const FVector2D TextSize = FontMeasureService->Measure(FrameString, SmallLayoutFont);
+
+			// Flip the text position if getting near the end of the view range
+			constexpr float TextOffsetPx = 2.f;
+			const bool  bDrawLeft = (static_cast<float>(AllottedGeometry.Size.X) - HandleEnd) < (TextSize.X + 14.f) - TextOffsetPx;
+			const float TextPosition = bDrawLeft ? HandleStart - static_cast<float>(TextSize.X) - TextOffsetPx : HandleEnd + TextOffsetPx;
+
+			FVector2D TextOffset(TextPosition, Args.bMirrorLabels ? TextSize.Y - 6.f : Args.AllottedGeometry.Size.Y - (Args.MajorTickHeight + TextSize.Y));
+
+			FSlateDrawElement::MakeText(
+				OutDrawElements,
+				Args.StartLayer + 1,
+				Args.AllottedGeometry.ToPaintGeometry(TextSize, FSlateLayoutTransform(TextOffset)),
+				FrameString,
+				SmallLayoutFont,
+				Args.DrawEffects,
+				Args.TickColor
+			);
+		}
+
+		return ArrowLayer;
+	}
+
+	return LayerId;
+}
+
+int32 FGeometryCacheTimeSlideController::OnPaintViewArea(const FGeometry& AllottedGeometry, const FSlateRect& MyCullingRect, FSlateWindowElementList& OutDrawElements, int32 LayerId, bool bEnabled, const FPaintViewAreaArgs& Args) const
+{
+	const ESlateDrawEffect DrawEffects = bEnabled ? ESlateDrawEffect::None : ESlateDrawEffect::DisabledEffect;
+
+	TRange<double> LocalViewRange = TimeSliderArgs.ViewRange.Get();
+	FScrubRangeToScreen RangeToScreen(LocalViewRange, AllottedGeometry.Size);
+
+	if (Args.PlaybackRangeArgs.IsSet())
+	{
+		FPaintPlaybackRangeArgs PaintArgs = Args.PlaybackRangeArgs.GetValue();
+		LayerId = DrawPlaybackRange(AllottedGeometry, MyCullingRect, OutDrawElements, LayerId, RangeToScreen, PaintArgs);
+		PaintArgs.SolidFillOpacity = 0.2f;
+		LayerId = DrawSelectionRange(AllottedGeometry, MyCullingRect, OutDrawElements, LayerId, RangeToScreen, PaintArgs);
+	}
+
+	if (Args.bDisplayTickLines)
+	{
+		static FLinearColor TickColor(0.1f, 0.1f, 0.1f, 0.3f);
+
+		// Draw major tick lines in the section area
+		FDrawTickArgs DrawTickArgs;
+		{
+			DrawTickArgs.AllottedGeometry = AllottedGeometry;
+			DrawTickArgs.bMirrorLabels = false;
+			DrawTickArgs.bOnlyDrawMajorTicks = true;
+			DrawTickArgs.TickColor = TickColor;
+			DrawTickArgs.CullingRect = MyCullingRect;
+			DrawTickArgs.DrawEffects = DrawEffects;
+			// Draw major ticks under sections
+			DrawTickArgs.StartLayer = LayerId - 1;
+			// Draw the tick the entire height of the section area
+			DrawTickArgs.TickOffset = 0.0f;
+			DrawTickArgs.MajorTickHeight = static_cast<float>(AllottedGeometry.Size.Y);
+		}
+
+		DrawTicks(OutDrawElements, LocalViewRange, RangeToScreen, DrawTickArgs);
+	}
+
+	if (Args.bDisplayScrubPosition)
+	{
+		// Draw a line for the scrub position
+		const float LinePos = RangeToScreen.InputToLocalX(TimeSliderArgs.ScrubPosition.Get().AsDecimal() / GetTickResolution().AsDecimal());
+
+		TArray<FVector2D> LinePoints;
+		{
+			LinePoints.AddUninitialized(2);
+			LinePoints[0] = FVector2D(0.0f, 0.0f);
+			LinePoints[1] = FVector2D(0.0f, FMath::FloorToFloat(AllottedGeometry.Size.Y));
+		}
+
+		FSlateDrawElement::MakeLines(
+			OutDrawElements,
+			LayerId + 1,
+			AllottedGeometry.ToPaintGeometry(FVector2f(1.0f, 1.0f), FSlateLayoutTransform(FVector2f(LinePos, 0.0f))),
+			LinePoints,
+			DrawEffects,
+			FLinearColor(1.f, 1.f, 1.f, .5f),
+			false
+		);
+	}
+
+	return LayerId;
+}
+
+FReply FGeometryCacheTimeSlideController::OnMouseButtonDown(SWidget& WidgetOwner, const FGeometry& MyGeometry, const FPointerEvent& MouseEvent)
+{
+	DistanceDragged = 0;
+	MouseDownPosition[0] = MouseDownPosition[1] = MouseEvent.GetScreenSpacePosition();
+	MouseDownGeometry = MyGeometry;	
+	return FReply::Unhandled();
+}
+
+FReply FGeometryCacheTimeSlideController::OnMouseButtonUp(SWidget& WidgetOwner, const FGeometry& MyGeometry, const FPointerEvent& MouseEvent)
+{
+	bool bHandleLeftMouseButton = MouseEvent.GetEffectingButton() == EKeys::LeftMouseButton && WidgetOwner.HasMouseCapture();
+	bool bHandleRightMouseButton = MouseEvent.GetEffectingButton() == EKeys::RightMouseButton && WidgetOwner.HasMouseCapture() && TimeSliderArgs.AllowZoom;
+
+	FScrubRangeToScreen RangeToScreen = FScrubRangeToScreen(TimeSliderArgs.ViewRange.Get(), MyGeometry.Size);
+	FFrameTime          MouseTime = ComputeFrameTimeFromMouse(MyGeometry, MouseEvent.GetScreenSpacePosition(), RangeToScreen);
+
+	if (bHandleRightMouseButton)
+	{
+		if (!bPanning && DistanceDragged == 0.0f)
+		{
+			return WeakTimeline.Pin()->OnMouseButtonUp(MyGeometry, MouseEvent).ReleaseMouseCapture();
+		}
+
+		bPanning = false;
+		DistanceDragged = 0.f;
+
+		return FReply::Handled().ReleaseMouseCapture();
+	}
+	else if (bHandleLeftMouseButton)
+	{
+		if (MouseDragType == DRAG_PLAYBACK_START)
+		{
+			TimeSliderArgs.OnPlaybackRangeEndDrag.ExecuteIfBound();
+		}
+		else if (MouseDragType == DRAG_PLAYBACK_END)
+		{
+			TimeSliderArgs.OnPlaybackRangeEndDrag.ExecuteIfBound();
+		}
+		else if (MouseDragType == DRAG_SELECTION_START)
+		{
+			TimeSliderArgs.OnSelectionRangeEndDrag.ExecuteIfBound();
+		}
+		else if (MouseDragType == DRAG_SELECTION_END)
+		{
+			TimeSliderArgs.OnSelectionRangeEndDrag.ExecuteIfBound();
+		}
+		else if (MouseDragType == DRAG_SETTING_RANGE)
+		{
+			// Zooming
+			FFrameTime MouseDownStart = ComputeFrameTimeFromMouse(MyGeometry, MouseDownPosition[0], RangeToScreen);
+
+			const bool bCanZoomIn = MouseTime > MouseDownStart;
+			const bool bCanZoomOut = ViewRangeStack.Num() > 0;
+			if (bCanZoomIn || bCanZoomOut)
+			{
+				TRange<double> ViewRange = TimeSliderArgs.ViewRange.Get();
+				if (!bCanZoomIn)
+				{
+					ViewRange = ViewRangeStack.Pop();
+				}
+
+				if (bCanZoomIn)
+				{
+					// push the current value onto the stack
+					ViewRangeStack.Add(ViewRange);
+
+					ViewRange = TRange<double>(MouseDownStart.FrameNumber / GetTickResolution(), MouseTime.FrameNumber / GetTickResolution());
+				}
+
+				TimeSliderArgs.OnViewRangeChanged.ExecuteIfBound(ViewRange, EViewRangeInterpolation::Immediate);
+				if (!TimeSliderArgs.ViewRange.IsBound())
+				{
+					// The output is not bound to a delegate so we'll manage the value ourselves
+					TimeSliderArgs.ViewRange.Set(ViewRange);
+				}
+			}
+		}
+		else
+		{
+			TimeSliderArgs.OnEndScrubberMovement.ExecuteIfBound();
+
+			CommitScrubPosition(MouseTime, /*bIsScrubbing=*/false);
+		}
+
+		MouseDragType = DRAG_NONE;
+		DistanceDragged = 0.f;
+
+		return FReply::Handled().ReleaseMouseCapture();
+	}
+
+	return FReply::Unhandled();
+}
+
+FReply FGeometryCacheTimeSlideController::OnMouseMove(SWidget& WidgetOwner, const FGeometry& MyGeometry, const FPointerEvent& MouseEvent)
+{
+	const bool bHandleLeftMouseButton = MouseEvent.IsMouseButtonDown(EKeys::LeftMouseButton);
+	const bool bHandleRightMouseButton = MouseEvent.IsMouseButtonDown(EKeys::RightMouseButton) && TimeSliderArgs.AllowZoom;
+
+	if (bHandleRightMouseButton)
+	{
+		if (!bPanning)
+		{
+			DistanceDragged += static_cast<float>(FMath::Abs(MouseEvent.GetCursorDelta().X));
+			if (DistanceDragged > FSlateApplication::Get().GetDragTriggerDistance())
+			{
+				bPanning = true;
+			}
+		}
+		else
+		{
+			const TRange<double> LocalViewRange = TimeSliderArgs.ViewRange.Get();
+			const double LocalViewRangeMin = LocalViewRange.GetLowerBoundValue();
+			const double LocalViewRangeMax = LocalViewRange.GetUpperBoundValue();
+
+			const FScrubRangeToScreen ScaleInfo(LocalViewRange, MyGeometry.Size);
+			const FVector2D ScreenDelta = MouseEvent.GetCursorDelta();
+			FVector2D InputDelta;
+			InputDelta.X = ScreenDelta.X / ScaleInfo.PixelsPerInput;
+
+			double NewViewOutputMin = LocalViewRangeMin - InputDelta.X;
+			double NewViewOutputMax = LocalViewRangeMax - InputDelta.X;
+
+			ClampViewRange(NewViewOutputMin, NewViewOutputMax);
+			SetViewRange(NewViewOutputMin, NewViewOutputMax, EViewRangeInterpolation::Immediate);
+		}
+	}
+	else if (bHandleLeftMouseButton)
+	{
+		const TRange<double> LocalViewRange = TimeSliderArgs.ViewRange.Get();
+		const FScrubRangeToScreen RangeToScreen(LocalViewRange, MyGeometry.Size);
+		DistanceDragged += static_cast<float>(FMath::Abs(MouseEvent.GetCursorDelta().X));
+
+		if (MouseDragType == DRAG_NONE)
+		{
+			if (DistanceDragged > FSlateApplication::Get().GetDragTriggerDistance())
+			{
+				const FFrameTime MouseDownFree = ComputeFrameTimeFromMouse(MyGeometry, MouseDownPosition[0], RangeToScreen, false);
+
+				const FFrameRate FrameResolution = GetTickResolution();
+				const bool       bLockedPlayRange = TimeSliderArgs.IsPlaybackRangeLocked.Get();
+				const float      MouseDownPixel = RangeToScreen.InputToLocalX(MouseDownFree / FrameResolution);
+				const bool       bHitScrubber = GetHitTestScrubberPixelRange(TimeSliderArgs.ScrubPosition.Get(), RangeToScreen).HandleRange.Contains(MouseDownPixel);
+				
+				const TRange<double>   SelectionRange = TimeSliderArgs.SelectionRange.Get() / FrameResolution;
+				const TRange<double>   PlaybackRange = TimeSliderArgs.PlaybackRange.Get() / FrameResolution;
+
+				// Disable selection range test if it's empty so that the playback range scrubbing gets priority
+				if (!SelectionRange.IsEmpty() && !bHitScrubber && HitTestRangeEnd(RangeToScreen, SelectionRange, MouseDownPixel))
+				{
+					// selection range end scrubber
+					MouseDragType = DRAG_SELECTION_END;
+					TimeSliderArgs.OnSelectionRangeBeginDrag.ExecuteIfBound();
+				}
+				else if (!SelectionRange.IsEmpty() && !bHitScrubber && HitTestRangeStart(RangeToScreen, SelectionRange, MouseDownPixel))
+				{
+					// selection range start scrubber
+					MouseDragType = DRAG_SELECTION_START;
+					TimeSliderArgs.OnSelectionRangeBeginDrag.ExecuteIfBound();
+				}
+				else if (!bLockedPlayRange && !bHitScrubber && HitTestRangeEnd(RangeToScreen, PlaybackRange, MouseDownPixel))
+				{
+					// playback range end scrubber
+					MouseDragType = DRAG_PLAYBACK_END;
+					TimeSliderArgs.OnPlaybackRangeBeginDrag.ExecuteIfBound();
+				}
+				else if (!bLockedPlayRange && !bHitScrubber && HitTestRangeStart(RangeToScreen, PlaybackRange, MouseDownPixel))
+				{
+					// playback range start scrubber
+					MouseDragType = DRAG_PLAYBACK_START;
+					TimeSliderArgs.OnPlaybackRangeBeginDrag.ExecuteIfBound();
+				}
+				else if (FSlateApplication::Get().GetModifierKeys().AreModifersDown(EModifierKey::Control))
+				{
+					MouseDragType = DRAG_SETTING_RANGE;
+				}
+				else
+				{
+					MouseDragType = DRAG_SCRUBBING_TIME;
+					TimeSliderArgs.OnBeginScrubberMovement.ExecuteIfBound();
+				}
+			}
+		}
+		else
+		{
+			FFrameTime MouseTime = ComputeFrameTimeFromMouse(MyGeometry, MouseEvent.GetScreenSpacePosition(), RangeToScreen);
+
+			// Set the start range time?
+			if (MouseDragType == DRAG_PLAYBACK_START)
+			{
+				SetPlaybackRangeStart(MouseTime.FrameNumber);
+			}
+			// Set the end range time?
+			else if (MouseDragType == DRAG_PLAYBACK_END)
+			{
+				SetPlaybackRangeEnd(MouseTime.FrameNumber - 1);
+			}
+			else if (MouseDragType == DRAG_SELECTION_START)
+			{
+				SetSelectionRangeStart(MouseTime.FrameNumber);
+			}
+			// Set the end range time?
+			else if (MouseDragType == DRAG_SELECTION_END)
+			{
+				SetSelectionRangeEnd(MouseTime.FrameNumber);
+			}
+			else if (MouseDragType == DRAG_SCRUBBING_TIME)
+			{
+				// Delegate responsibility for clamping to the current viewrange to the client
+				CommitScrubPosition(MouseTime, /*bIsScrubbing=*/true);
+			}
+			else if (MouseDragType == DRAG_SETTING_RANGE)
+			{
+				MouseDownPosition[1] = MouseEvent.GetScreenSpacePosition();
+			}
+			else if (MouseDragType == DRAG_TIME)
+			{
+				return FReply::Unhandled();
+			}
+		}
+	}
+
+	if (DistanceDragged != 0.f && (bHandleLeftMouseButton || bHandleRightMouseButton))
+	{
+		return FReply::Handled().CaptureMouse(WidgetOwner.AsShared());
+	}
+
+	return FReply::Handled();
+
+}
+
+FReply FGeometryCacheTimeSlideController::OnMouseWheel(SWidget& WidgetOwner, const FGeometry& MyGeometry, const FPointerEvent& MouseEvent)
+{
+	TOptional<TRange<float>> NewTargetRange;
+
+	if (TimeSliderArgs.AllowZoom && MouseEvent.IsControlDown())
+	{
+		const float MouseFractionX = static_cast<float>(MyGeometry.AbsoluteToLocal(MouseEvent.GetScreenSpacePosition()).X / MyGeometry.GetLocalSize().X);
+
+		const float ZoomDelta = -0.2f * MouseEvent.GetWheelDelta();
+		if (ZoomByDelta(ZoomDelta, MouseFractionX))
+		{
+			return FReply::Handled();
+		}
+	}
+	else if (MouseEvent.IsShiftDown())
+	{
+		PanByDelta(-MouseEvent.GetWheelDelta());
+		return FReply::Handled();
+	}
+
+	return FReply::Unhandled();
+}
+
+FCursorReply FGeometryCacheTimeSlideController::OnCursorQuery(TSharedRef<const SWidget> WidgetOwner, const FGeometry& MyGeometry, const FPointerEvent& CursorEvent) const
+{
+	const FScrubRangeToScreen RangeToScreen(TimeSliderArgs.ViewRange.Get(), MyGeometry.Size);
+
+	const FFrameRate FrameResolution = GetTickResolution();
+	const bool       bLockedPlayRange = TimeSliderArgs.IsPlaybackRangeLocked.Get();
+	const float      HitTestPixel = static_cast<float>(MyGeometry.AbsoluteToLocal(CursorEvent.GetScreenSpacePosition()).X);
+	const bool       bHitScrubber = GetHitTestScrubberPixelRange(TimeSliderArgs.ScrubPosition.Get(), RangeToScreen).HandleRange.Contains(HitTestPixel);
+	
+	const TRange<double>   SelectionRange = TimeSliderArgs.SelectionRange.Get() / FrameResolution;
+	const TRange<double>   PlaybackRange = TimeSliderArgs.PlaybackRange.Get() / FrameResolution;
+
+	if (MouseDragType == DRAG_SCRUBBING_TIME)
+	{
+		return FCursorReply::Unhandled();
+	}
+
+	// Use L/R resize cursor if we're dragging or hovering a playback range bound
+	if ((MouseDragType == DRAG_PLAYBACK_END) ||
+		(MouseDragType == DRAG_PLAYBACK_START) ||
+		(MouseDragType == DRAG_SELECTION_START) ||
+		(MouseDragType == DRAG_SELECTION_END) ||
+		(!bLockedPlayRange && !bHitScrubber && HitTestRangeStart(RangeToScreen, PlaybackRange, HitTestPixel)) ||
+		(!bLockedPlayRange && !bHitScrubber && HitTestRangeEnd(RangeToScreen, PlaybackRange, HitTestPixel)) ||
+		(!SelectionRange.IsEmpty() && !bHitScrubber && HitTestRangeStart(RangeToScreen, SelectionRange, HitTestPixel)) ||
+		(!SelectionRange.IsEmpty() && !bHitScrubber && HitTestRangeEnd(RangeToScreen, SelectionRange, HitTestPixel)))
+	{
+		return FCursorReply::Cursor(EMouseCursor::ResizeLeftRight);
+	}
+
+	return FCursorReply::Unhandled();
+}
+
+void FGeometryCacheTimeSlideController::SetViewRange(double NewRangeMin, double NewRangeMax, EViewRangeInterpolation Interpolation)
+{
+	// Clamp to a minimum size to avoid zero-sized or negative visible ranges
+	double MinVisibleTimeRange = FFrameNumber(1) / GetTickResolution();
+	TRange<double> ExistingViewRange = TimeSliderArgs.ViewRange.Get();
+
+	if (NewRangeMax == ExistingViewRange.GetUpperBoundValue())
+	{
+		if (NewRangeMin > NewRangeMax - MinVisibleTimeRange)
+		{
+			NewRangeMin = NewRangeMax - MinVisibleTimeRange;
+		}
+	}
+	else if (NewRangeMax < NewRangeMin + MinVisibleTimeRange)
+	{
+		NewRangeMax = NewRangeMin + MinVisibleTimeRange;
+	}
+
+	// Clamp to the clamp range
+	const TRange<double> NewRange = TRange<double>(NewRangeMin, NewRangeMax);
+	TimeSliderArgs.OnViewRangeChanged.ExecuteIfBound(NewRange, Interpolation);
+
+	if (!TimeSliderArgs.ViewRange.IsBound())
+	{
+		// The  output is not bound to a delegate so we'll manage the value ourselves (no animation)
+		TimeSliderArgs.ViewRange.Set(NewRange);
+	}
+}
+
+void FGeometryCacheTimeSlideController::SetClampRange(double NewRangeMin, double NewRangeMax)
+{
+	const TRange<double> NewRange(NewRangeMin, NewRangeMax);
+
+	TimeSliderArgs.OnClampRangeChanged.ExecuteIfBound(NewRange);
+
+	if (!TimeSliderArgs.ClampRange.IsBound())
+	{
+		// The  output is not bound to a delegate so we'll manage the value ourselves (no animation)
+		TimeSliderArgs.ClampRange.Set(NewRange);
+	}
+}
+
+void FGeometryCacheTimeSlideController::SetPlayRange(FFrameNumber RangeStart, int32 RangeDuration)
+{
+	check(RangeDuration >= 0);
+
+	const TRange<FFrameNumber> NewRange(RangeStart, RangeStart + RangeDuration);
+
+	TimeSliderArgs.OnPlaybackRangeChanged.ExecuteIfBound(NewRange);
+
+	if (!TimeSliderArgs.PlaybackRange.IsBound())
+	{
+		// The  output is not bound to a delegate so we'll manage the value ourselves (no animation)
+		TimeSliderArgs.PlaybackRange.Set(NewRange);
+	}
+}
+
+void FGeometryCacheTimeSlideController::ClampViewRange(double& NewRangeMin, double& NewRangeMax)
+{
+	bool bNeedsClampSet = false;
+	double NewClampRangeMin = TimeSliderArgs.ClampRange.Get().GetLowerBoundValue();
+	if (NewRangeMin < TimeSliderArgs.ClampRange.Get().GetLowerBoundValue())
+	{
+		NewClampRangeMin = NewRangeMin;
+		bNeedsClampSet = true;
+	}
+
+	double NewClampRangeMax = TimeSliderArgs.ClampRange.Get().GetUpperBoundValue();
+	if (NewRangeMax > TimeSliderArgs.ClampRange.Get().GetUpperBoundValue())
+	{
+		NewClampRangeMax = NewRangeMax;
+		bNeedsClampSet = true;
+	}
+
+	if (bNeedsClampSet)
+	{
+		SetClampRange(NewClampRangeMin, NewClampRangeMax);
+	}
+}
+
+bool FGeometryCacheTimeSlideController::ZoomByDelta(float InDelta, float MousePositionFraction)
+{
+	TRange<double> LocalViewRange = TimeSliderArgs.ViewRange.Get().GetAnimationTarget();
+	double LocalViewRangeMax = LocalViewRange.GetUpperBoundValue();
+	double LocalViewRangeMin = LocalViewRange.GetLowerBoundValue();
+	const double OutputViewSize = LocalViewRangeMax - LocalViewRangeMin;
+	const double OutputChange = OutputViewSize * InDelta;
+
+	double NewViewOutputMin = LocalViewRangeMin - (OutputChange * MousePositionFraction);
+	double NewViewOutputMax = LocalViewRangeMax + (OutputChange * (1.f - MousePositionFraction));
+
+	if (NewViewOutputMin < NewViewOutputMax)
+	{
+		ClampViewRange(NewViewOutputMin, NewViewOutputMax);
+		SetViewRange(NewViewOutputMin, NewViewOutputMax, EViewRangeInterpolation::Animated);
+		return true;
+	}
+
+	return false;
+}
+
+void FGeometryCacheTimeSlideController::PanByDelta(float InDelta)
+{
+	const TRange<double> LocalViewRange = TimeSliderArgs.ViewRange.Get().GetAnimationTarget();
+
+	const double CurrentMin = LocalViewRange.GetLowerBoundValue();
+	const double CurrentMax = LocalViewRange.GetUpperBoundValue();
+
+	// Adjust the delta to be a percentage of the current range
+	InDelta *= ScrubConstants::ScrollPanFraction * static_cast<float>(CurrentMax - CurrentMin);
+
+	double NewViewOutputMin = CurrentMin + InDelta;
+	double NewViewOutputMax = CurrentMax + InDelta;
+
+	ClampViewRange(NewViewOutputMin, NewViewOutputMax);
+	SetViewRange(NewViewOutputMin, NewViewOutputMax, EViewRangeInterpolation::Animated);
+}
+
+FFrameTime FGeometryCacheTimeSlideController::GetFrameTimeFromMouse(const FGeometry& Geometry, FVector2D ScreenSpacePosition) const
+{
+	FScrubRangeToScreen ScrubRangeToScreen(TimeSliderArgs.ViewRange.Get(), Geometry.Size);
+	return ComputeFrameTimeFromMouse(Geometry, ScreenSpacePosition, ScrubRangeToScreen);
+}
+
+void FGeometryCacheTimeSlideController::CommitScrubPosition(FFrameTime NewValue, bool bIsScrubbing)
+{
+	// Manage the scrub position ourselves if its not bound to a delegate
+	if (!TimeSliderArgs.ScrubPosition.IsBound())
+	{
+		TimeSliderArgs.ScrubPosition.Set(NewValue);
+	}
+
+	TimeSliderArgs.OnScrubPositionChanged.ExecuteIfBound(NewValue, bIsScrubbing, /*bEvaluate*/ true); 
+}
+
+void FGeometryCacheTimeSlideController::DrawTicks(FSlateWindowElementList& OutDrawElements, const TRange<double>& ViewRange, const FScrubRangeToScreen& RangeToScreen, FDrawTickArgs& InArgs) const
+{
+	TSharedPtr<SGeometryCacheTimeline> Timeline = WeakTimeline.Pin();
+	if (!Timeline.IsValid())
+	{
+		return;
+	}
+
+	if (!FMath::IsFinite(ViewRange.GetLowerBoundValue()) || !FMath::IsFinite(ViewRange.GetUpperBoundValue()))
+	{
+		return;
+	}
+
+	FFrameRate     FrameResolution = GetTickResolution();
+	FPaintGeometry PaintGeometry = InArgs.AllottedGeometry.ToPaintGeometry();
+	FSlateFontInfo SmallLayoutFont = FCoreStyle::GetDefaultFontStyle("Regular", 8);
+
+	double MajorGridStep = 0.0;
+	int32  MinorDivisions = 0;
+	if (!Timeline->GetGridMetrics(static_cast<float>(InArgs.AllottedGeometry.Size.X), MajorGridStep, MinorDivisions))
+	{
+		return;
+	}
+
+	if (InArgs.bOnlyDrawMajorTicks)
+	{
+		MinorDivisions = 0;
+	}
+
+	TArray<FVector2D> LinePoints;
+	LinePoints.SetNumUninitialized(2);
+
+	const bool bAntiAliasLines = false;
+
+	const double FirstMajorLine = FMath::FloorToDouble(ViewRange.GetLowerBoundValue() / MajorGridStep) * MajorGridStep;
+	const double LastMajorLine = FMath::CeilToDouble(ViewRange.GetUpperBoundValue() / MajorGridStep) * MajorGridStep;
+
+	for (double CurrentMajorLine = FirstMajorLine; CurrentMajorLine < LastMajorLine; CurrentMajorLine += MajorGridStep)
+	{
+		float MajorLinePx = RangeToScreen.InputToLocalX(CurrentMajorLine);
+
+		LinePoints[0] = FVector2D(MajorLinePx, InArgs.TickOffset);
+		LinePoints[1] = FVector2D(MajorLinePx, InArgs.TickOffset + InArgs.MajorTickHeight);
+
+		// Draw each tick mark
+		FSlateDrawElement::MakeLines(
+			OutDrawElements,
+			InArgs.StartLayer,
+			PaintGeometry,
+			LinePoints,
+			InArgs.DrawEffects,
+			InArgs.TickColor,
+			bAntiAliasLines
+		);
+
+		if (!InArgs.bOnlyDrawMajorTicks)
+		{
+			FString FrameString = TimeSliderArgs.NumericTypeInterface->ToString((CurrentMajorLine * FrameResolution).RoundToFrame().Value);
+
+			// Space the text between the tick mark but slightly above
+			FVector2D TextOffset(MajorLinePx + 5.f, InArgs.bMirrorLabels ? 3.f : FMath::Abs(InArgs.AllottedGeometry.Size.Y - (InArgs.MajorTickHeight + 3.f)));
+			FSlateDrawElement::MakeText(
+				OutDrawElements,
+				InArgs.StartLayer + 1,
+				InArgs.AllottedGeometry.ToPaintGeometry(InArgs.AllottedGeometry.Size, FSlateLayoutTransform(TextOffset)),
+				FrameString,
+				SmallLayoutFont,
+				InArgs.DrawEffects,
+				InArgs.TickColor * 0.65f
+			);
+		}
+
+		for (int32 Step = 1; Step < MinorDivisions; ++Step)
+		{
+			// Compute the size of each tick mark.  If we are half way between to visible values display a slightly larger tick mark
+			const float MinorTickHeight = ((MinorDivisions % 2 == 0) && (Step % (MinorDivisions / 2)) == 0) ? 6.0f : 2.0f;
+			const float MinorLinePx = RangeToScreen.InputToLocalX(CurrentMajorLine + Step * MajorGridStep / MinorDivisions);
+
+			LinePoints[0] = FVector2D(MinorLinePx, InArgs.bMirrorLabels ? 0.0f : FMath::Abs(InArgs.AllottedGeometry.Size.Y - MinorTickHeight));
+			LinePoints[1] = FVector2D(MinorLinePx, LinePoints[0].Y + MinorTickHeight);
+
+			// Draw each sub mark
+			FSlateDrawElement::MakeLines(
+				OutDrawElements,
+				InArgs.StartLayer,
+				PaintGeometry,
+				LinePoints,
+				InArgs.DrawEffects,
+				InArgs.TickColor,
+				bAntiAliasLines
+			);
+		}
+	}
+}
+
+int32 FGeometryCacheTimeSlideController::DrawSelectionRange(const FGeometry& AllottedGeometry, const FSlateRect& MyCullingRect, FSlateWindowElementList& OutDrawElements, int32 LayerId, const FScrubRangeToScreen& RangeToScreen, const FPaintPlaybackRangeArgs& Args) const
+{
+	TRange<double> SelectionRange = TimeSliderArgs.SelectionRange.Get() / GetTickResolution();
+
+	if (!SelectionRange.IsEmpty() && SelectionRange.HasLowerBound() && SelectionRange.HasUpperBound())
+	{
+		const float SelectionRangeL = RangeToScreen.InputToLocalX(SelectionRange.GetLowerBoundValue()) - 1;
+		const float SelectionRangeR = RangeToScreen.InputToLocalX(SelectionRange.GetUpperBoundValue()) + 1;
+		const FLinearColor DrawColor = FAppStyle::GetSlateColor("SelectionColor").GetColor(FWidgetStyle());
+
+		if (Args.SolidFillOpacity > 0.f)
+		{
+			FSlateDrawElement::MakeBox(
+				OutDrawElements,
+				LayerId + 1,
+				AllottedGeometry.ToPaintGeometry(FVector2f(SelectionRangeR - SelectionRangeL, AllottedGeometry.Size.Y), FSlateLayoutTransform(FVector2D(SelectionRangeL, 0.f))),
+				FAppStyle::GetBrush("WhiteBrush"),
+				ESlateDrawEffect::None,
+				DrawColor.CopyWithNewOpacity(Args.SolidFillOpacity)
+			);
+		}
+
+		FSlateDrawElement::MakeBox(
+			OutDrawElements,
+			LayerId + 1,
+			AllottedGeometry.ToPaintGeometry(FVector2f(Args.BrushWidth, AllottedGeometry.Size.Y), FSlateLayoutTransform(FVector2D(SelectionRangeL, 0.f))),
+			Args.StartBrush,
+			ESlateDrawEffect::None,
+			DrawColor
+		);
+
+		FSlateDrawElement::MakeBox(
+			OutDrawElements,
+			LayerId + 1,
+			AllottedGeometry.ToPaintGeometry(FVector2f(Args.BrushWidth, AllottedGeometry.Size.Y), FSlateLayoutTransform(FVector2D(SelectionRangeR - Args.BrushWidth, 0.f))),
+			Args.EndBrush,
+			ESlateDrawEffect::None,
+			DrawColor
+		);
+	}
+
+	return LayerId + 1;
+}
+
+int32 FGeometryCacheTimeSlideController::DrawPlaybackRange(const FGeometry& AllottedGeometry, const FSlateRect& MyCullingRect, FSlateWindowElementList& OutDrawElements, int32 LayerId, const FScrubRangeToScreen& RangeToScreen, const FPaintPlaybackRangeArgs& Args) const
+{
+	if (!TimeSliderArgs.PlaybackRange.IsSet())
+	{
+		return LayerId;
+	}
+
+	const uint8 OpacityBlend = TimeSliderArgs.SubSequenceRange.Get().IsSet() ? 128 : 255;
+
+	TRange<FFrameNumber> PlaybackRange = TimeSliderArgs.PlaybackRange.Get();
+	FFrameRate TickResolution = GetTickResolution();
+	const float PlaybackRangeL = RangeToScreen.InputToLocalX(PlaybackRange.GetLowerBoundValue() / TickResolution);
+	const float PlaybackRangeR = RangeToScreen.InputToLocalX(PlaybackRange.GetUpperBoundValue() / TickResolution) - 1;
+
+	FSlateDrawElement::MakeBox(
+		OutDrawElements,
+		LayerId + 1,
+		AllottedGeometry.ToPaintGeometry(FVector2f(Args.BrushWidth, AllottedGeometry.Size.Y), FSlateLayoutTransform(FVector2D(PlaybackRangeL, 0.f))),
+		Args.StartBrush,
+		ESlateDrawEffect::None,
+		FColor(32, 128, 32, OpacityBlend)	// 120, 75, 50 (HSV)
+	);
+
+	FSlateDrawElement::MakeBox(
+		OutDrawElements,
+		LayerId + 1,
+		AllottedGeometry.ToPaintGeometry(FVector2f(Args.BrushWidth, AllottedGeometry.Size.Y), FSlateLayoutTransform(FVector2D(PlaybackRangeR - Args.BrushWidth, 0.f))),
+		Args.EndBrush,
+		ESlateDrawEffect::None,
+		FColor(128, 32, 32, OpacityBlend)	// 0, 75, 50 (HSV)
+	);
+
+	// Black tint for excluded regions
+	FSlateDrawElement::MakeBox(
+		OutDrawElements,
+		LayerId + 1,
+		AllottedGeometry.ToPaintGeometry(FVector2f(PlaybackRangeL, AllottedGeometry.Size.Y), FSlateLayoutTransform(FVector2D(0.f, 0.f))),
+		FAppStyle::GetBrush("WhiteBrush"),
+		ESlateDrawEffect::None,
+		FLinearColor::Black.CopyWithNewOpacity(0.3f * OpacityBlend / 255.f)
+	);
+
+	FSlateDrawElement::MakeBox(
+		OutDrawElements,
+		LayerId + 1,
+		AllottedGeometry.ToPaintGeometry(FVector2f(AllottedGeometry.Size.X - PlaybackRangeR, AllottedGeometry.Size.Y), FSlateLayoutTransform(FVector2D(PlaybackRangeR, 0.f))),
+		FAppStyle::GetBrush("WhiteBrush"),
+		ESlateDrawEffect::None,
+		FLinearColor::Black.CopyWithNewOpacity(0.3f * OpacityBlend / 255.f)
+	);
+
+	return LayerId + 1;
+}
+
+bool FGeometryCacheTimeSlideController::HitTestRangeStart(const FScrubRangeToScreen& RangeToScreen, const TRange<double>& Range, float HitPixel) const
+{
+	if (Range.HasLowerBound())
+	{
+		static float BrushSizeInStateUnits = 6.f, DragToleranceSlateUnits = 2.f, MouseTolerance = 2.f;
+		const float  RangeStartPixel = RangeToScreen.InputToLocalX(Range.GetLowerBoundValue());
+
+		// Hit test against the brush region to the right of the playback start position, +/- DragToleranceSlateUnits
+		return HitPixel >= RangeStartPixel - MouseTolerance - DragToleranceSlateUnits &&
+			HitPixel <= RangeStartPixel + MouseTolerance + BrushSizeInStateUnits + DragToleranceSlateUnits;
+	}
+
+	return false;
+}
+
+bool FGeometryCacheTimeSlideController::HitTestRangeEnd(const FScrubRangeToScreen& RangeToScreen, const TRange<double>& Range, float HitPixel) const
+{
+	if (Range.HasUpperBound())
+	{
+		static float BrushSizeInStateUnits = 6.f, DragToleranceSlateUnits = 2.f, MouseTolerance = 2.f;
+		const float  RangeEndPixel = RangeToScreen.InputToLocalX(Range.GetUpperBoundValue());
+
+		// Hit test against the brush region to the left of the playback end position, +/- DragToleranceSlateUnits
+		return HitPixel >= RangeEndPixel - MouseTolerance - BrushSizeInStateUnits - DragToleranceSlateUnits &&
+			HitPixel <= RangeEndPixel + MouseTolerance + DragToleranceSlateUnits;
+	}
+
+	return false;
+}
+
+void FGeometryCacheTimeSlideController::SetPlaybackRangeStart(FFrameNumber NewStart)
+{
+	TRange<FFrameNumber> PlaybackRange = TimeSliderArgs.PlaybackRange.Get();
+
+	if (NewStart <= UE::MovieScene::DiscreteExclusiveUpper(PlaybackRange))
+	{
+		TimeSliderArgs.OnPlaybackRangeChanged.ExecuteIfBound(TRange<FFrameNumber>(NewStart, PlaybackRange.GetUpperBound()));
+	}
+}
+
+void FGeometryCacheTimeSlideController::SetPlaybackRangeEnd(FFrameNumber NewEnd)
+{
+	TRange<FFrameNumber> PlaybackRange = TimeSliderArgs.PlaybackRange.Get();
+
+	if (NewEnd >= UE::MovieScene::DiscreteInclusiveLower(PlaybackRange))
+	{
+		TimeSliderArgs.OnPlaybackRangeChanged.ExecuteIfBound(TRange<FFrameNumber>(PlaybackRange.GetLowerBound(), NewEnd));
+	}
+}
+
+void FGeometryCacheTimeSlideController::SetSelectionRangeStart(FFrameNumber NewStart)
+{
+	TRange<FFrameNumber> SelectionRange = TimeSliderArgs.SelectionRange.Get();
+
+	if (SelectionRange.IsEmpty())
+	{
+		TimeSliderArgs.OnSelectionRangeChanged.ExecuteIfBound(TRange<FFrameNumber>(NewStart, NewStart + 1));
+	}
+	else if (NewStart <= UE::MovieScene::DiscreteExclusiveUpper(SelectionRange))
+	{
+		TimeSliderArgs.OnSelectionRangeChanged.ExecuteIfBound(TRange<FFrameNumber>(NewStart, SelectionRange.GetUpperBound()));
+	}
+}
+
+void FGeometryCacheTimeSlideController::SetSelectionRangeEnd(FFrameNumber NewEnd)
+{
+	TRange<FFrameNumber> SelectionRange = TimeSliderArgs.SelectionRange.Get();
+
+	if (SelectionRange.IsEmpty())
+	{
+		TimeSliderArgs.OnSelectionRangeChanged.ExecuteIfBound(TRange<FFrameNumber>(NewEnd - 1, NewEnd));
+	}
+	else if (NewEnd >= UE::MovieScene::DiscreteInclusiveLower(SelectionRange))
+	{
+		TimeSliderArgs.OnSelectionRangeChanged.ExecuteIfBound(TRange<FFrameNumber>(SelectionRange.GetLowerBound(), NewEnd));
+	}
+}
+
+FFrameTime FGeometryCacheTimeSlideController::ComputeFrameTimeFromMouse(const FGeometry& Geometry, FVector2D ScreenSpacePosition, FScrubRangeToScreen RangeToScreen, bool CheckSnapping) const
+{
+	const FVector2D CursorPos = Geometry.AbsoluteToLocal(ScreenSpacePosition);
+	const double    MouseValue = RangeToScreen.LocalXToInput(static_cast<float>(CursorPos.X));
+
+	return MouseValue * GetTickResolution();
+}
+
+FGeometryCacheTimeSlideController::FScrubPixelRange FGeometryCacheTimeSlideController::GetHitTestScrubberPixelRange(FFrameTime ScrubTime, const FScrubRangeToScreen& RangeToScreen) const
+{
+	constexpr float DragToleranceSlateUnits = 2.f, MouseTolerance = 2.f;
+	return GetScrubberPixelRange(ScrubTime, GetTickResolution(), GetDisplayRate(), RangeToScreen, DragToleranceSlateUnits + MouseTolerance);
+}
+
+FGeometryCacheTimeSlideController::FScrubPixelRange FGeometryCacheTimeSlideController::GetScrubberPixelRange(FFrameTime ScrubTime, const FScrubRangeToScreen& RangeToScreen) const
+{
+	return GetScrubberPixelRange(ScrubTime, GetTickResolution(), GetDisplayRate(), RangeToScreen);
+}
+
+FGeometryCacheTimeSlideController::FScrubPixelRange FGeometryCacheTimeSlideController::GetScrubberPixelRange(FFrameTime ScrubTime, FFrameRate Resolution, FFrameRate PlayRate, const FScrubRangeToScreen& RangeToScreen, float DilationPixels) const
+{
+	const FFrameNumber Frame = ScrubTime.FloorToFrame();// ConvertFrameTime(ScrubTime, Resolution, PlayRate).FloorToFrame();
+
+	float StartPixel = RangeToScreen.InputToLocalX(Frame / Resolution);
+	float EndPixel = RangeToScreen.InputToLocalX((Frame + 1) / Resolution);
+
+	{
+		const float RoundedStartPixel = static_cast<float>(FMath::RoundToInt(StartPixel));
+		EndPixel -= (StartPixel - RoundedStartPixel);
+
+		StartPixel = RoundedStartPixel;
+		EndPixel = FMath::Max(EndPixel, StartPixel + 1);
+	}
+
+	FScrubPixelRange Range;
+
+	float MinScrubSize = 14.f;
+	Range.bClamped = EndPixel - StartPixel < MinScrubSize;
+	Range.Range = TRange<float>(StartPixel, EndPixel);
+	if (Range.bClamped)
+	{
+		Range.HandleRange = TRange<float>(
+			(StartPixel + EndPixel - MinScrubSize) * .5f,
+			(StartPixel + EndPixel + MinScrubSize) * .5f);
+	}
+	else
+	{
+		Range.HandleRange = Range.Range;
+	}
+
+	return Range;
+}
\ No newline at end of file
diff --git a/Engine/Plugins/Runtime/GeometryCache/Source/GeometryCacheEd/Private/GeometryCacheTimeSliderController.h b/Engine/Plugins/Runtime/GeometryCache/Source/GeometryCacheEd/Private/GeometryCacheTimeSliderController.h
new file mode 100644
--- /dev/null
+++ b/Engine/Plugins/Runtime/GeometryCache/Source/GeometryCacheEd/Private/GeometryCacheTimeSliderController.h
@@ -0,0 +1,186 @@
+// Copyright Epic Games, Inc. All Rights Reserved.
+
+#pragma once
+
+#include "CoreMinimal.h"
+#include "Widgets/SWidget.h"
+#include "ITimeSlider.h"
+#include "TimeSliderArgs.h"
+#include "Input/CursorReply.h"
+#include "Input/Reply.h"
+
+class SGeometryCacheTimeline;
+struct FSlateBrush;
+
+class FGeometryCacheTimeSlideController : public ITimeSliderController, public TSharedFromThis<FGeometryCacheTimeSlideController>
+{
+	struct FDrawTickArgs;
+	struct FScrubRangeToScreen;
+
+public:
+	FGeometryCacheTimeSlideController(const FTimeSliderArgs& InArgs, TWeakPtr<SGeometryCacheTimeline> GeometryCacheTimeline);
+	~FGeometryCacheTimeSlideController();
+
+	/** ITimeSliderController Interface */
+	virtual int32 OnPaintTimeSlider(bool bMirrorLabels, const FGeometry& AllottedGeometry, const FSlateRect& MyCullingRect, FSlateWindowElementList& OutDrawElements, int32 LayerId, const FWidgetStyle& InWidgetStyle, bool bParentEnabled) const override;
+	virtual int32 OnPaintViewArea(const FGeometry& AllottedGeometry, const FSlateRect& MyCullingRect, FSlateWindowElementList& OutDrawElements, int32 LayerId, bool bEnabled, const FPaintViewAreaArgs& Args) const override;
+	virtual FReply OnMouseButtonDown(SWidget& WidgetOwner, const FGeometry& MyGeometry, const FPointerEvent& MouseEvent) override;
+	virtual FReply OnMouseButtonUp(SWidget& WidgetOwner, const FGeometry& MyGeometry, const FPointerEvent& MouseEvent) override;
+	virtual FReply OnMouseMove(SWidget& WidgetOwner, const FGeometry& MyGeometry, const FPointerEvent& MouseEvent) override;
+	virtual FReply OnMouseWheel(SWidget& WidgetOwner, const FGeometry& MyGeometry, const FPointerEvent& MouseEvent) override;
+	virtual FCursorReply OnCursorQuery(TSharedRef<const SWidget> WidgetOwner, const FGeometry& MyGeometry, const FPointerEvent& CursorEvent) const override;
+	virtual void SetViewRange(double NewRangeMin, double NewRangeMax, EViewRangeInterpolation Interpolation) override;
+	virtual void SetClampRange(double NewRangeMin, double NewRangeMax) override;
+	virtual void SetPlayRange(FFrameNumber RangeStart, int32 RangeDuration) override;
+	virtual FFrameRate GetDisplayRate() const override { return TimeSliderArgs.DisplayRate.Get(); }
+	virtual FFrameRate GetTickResolution() const override { return TimeSliderArgs.TickResolution.Get(); }
+	virtual FAnimatedRange GetViewRange() const override { return TimeSliderArgs.ViewRange.Get(); }
+	virtual FAnimatedRange GetClampRange() const override { return TimeSliderArgs.ClampRange.Get(); }
+	virtual TRange<FFrameNumber> GetPlayRange() const override { return TimeSliderArgs.PlaybackRange.Get(TRange<FFrameNumber>()); }
+	virtual FFrameTime GetScrubPosition() const { return TimeSliderArgs.ScrubPosition.Get(); }
+	virtual void SetScrubPosition(FFrameTime InTime, bool bEvaluate) { TimeSliderArgs.OnScrubPositionChanged.ExecuteIfBound(InTime, false, bEvaluate); }
+	/** End ITimeSliderController Interface */
+
+	/**
+	 * Clamp the given range to the clamp range
+	 *
+	 * @param NewRangeMin		The new lower bound of the range
+	 * @param NewRangeMax		The new upper bound of the range
+	 */
+	void ClampViewRange(double& NewRangeMin, double& NewRangeMax);
+
+	/**
+	 * Zoom the range by a given delta.
+	 *
+	 * @param InDelta					The total amount to zoom by (+ve = zoom out, -ve = zoom in)
+	 * @param MousePositionFraction		Bias to apply to lower/upper extents of the range. (0 = lower, 0.5 = equal, 1 = upper)
+	 */
+	bool ZoomByDelta(float InDelta, float MousePositionFraction = 0.5f);
+
+	/**
+	 * Pan the range by a given delta
+	 *
+	 * @param InDelta		The total amount to pan by (+ve = pan forwards in time, -ve = pan backwards in time)
+	 */
+	void PanByDelta(float InDelta);
+
+	/** Determine frame time from a mouse position */
+	FFrameTime GetFrameTimeFromMouse(const FGeometry& Geometry, FVector2D ScreenSpacePosition) const;
+
+private:
+	/**
+	 * Call this method when the user's interaction has changed the scrub position
+	 *
+	 * @param NewValue				Value resulting from the user's interaction
+	 * @param bIsScrubbing			True if done via scrubbing, false if just releasing scrubbing
+	 */
+	void CommitScrubPosition(FFrameTime NewValue, bool bIsScrubbing);
+
+	/**
+	 * Draw time tick marks
+	 *
+	 * @param OutDrawElements	List to add draw elements to
+	 * @param ViewRange			The currently visible time range in seconds
+	 * @param RangeToScreen		Time range to screen space converter
+	 * @param InArgs			Parameters for drawing the tick lines
+	 */
+	void DrawTicks(FSlateWindowElementList& OutDrawElements, const TRange<double>& ViewRange, const FScrubRangeToScreen& RangeToScreen, FDrawTickArgs& InArgs) const;
+
+	/**
+	 * Draw the selection range.
+	 *
+	 * @return The new layer ID.
+	 */
+	int32 DrawSelectionRange(const FGeometry& AllottedGeometry, const FSlateRect& MyCullingRect, FSlateWindowElementList& OutDrawElements, int32 LayerId, const FScrubRangeToScreen& RangeToScreen, const FPaintPlaybackRangeArgs& Args) const;
+
+	/**
+	 * Draw the playback range.
+	 *
+	 * @return the new layer ID
+	 */
+	int32 DrawPlaybackRange(const FGeometry& AllottedGeometry, const FSlateRect& MyCullingRect, FSlateWindowElementList& OutDrawElements, int32 LayerId, const FScrubRangeToScreen& RangeToScreen, const FPaintPlaybackRangeArgs& Args) const;
+
+	/**
+	 * Hit test the lower bound of a range
+	 */
+	bool HitTestRangeStart(const FScrubRangeToScreen& RangeToScreen, const TRange<double>& Range, float HitPixel) const;
+
+	/**
+	 * Hit test the upper bound of a range
+	 */
+	bool HitTestRangeEnd(const FScrubRangeToScreen& RangeToScreen, const TRange<double>& Range, float HitPixel) const;
+
+	void SetPlaybackRangeStart(FFrameNumber NewStart);
+	void SetPlaybackRangeEnd(FFrameNumber NewEnd);
+	
+	void SetSelectionRangeStart(FFrameNumber NewStart);
+	void SetSelectionRangeEnd(FFrameNumber NewEnd);
+
+	FFrameTime ComputeFrameTimeFromMouse(const FGeometry& Geometry, FVector2D ScreenSpacePosition, FScrubRangeToScreen RangeToScreen, bool CheckSnapping = true) const;
+
+private:
+	struct FScrubPixelRange
+	{
+		TRange<float> Range;
+		TRange<float> HandleRange;
+		bool bClamped;
+	};
+
+	FScrubPixelRange GetHitTestScrubberPixelRange(FFrameTime ScrubTime, const FScrubRangeToScreen& RangeToScreen) const;
+
+	FScrubPixelRange GetScrubberPixelRange(FFrameTime ScrubTime, const FScrubRangeToScreen& RangeToScreen) const;
+	FScrubPixelRange GetScrubberPixelRange(FFrameTime ScrubTime, FFrameRate Resolution, FFrameRate PlayRate, const FScrubRangeToScreen& RangeToScreen, float DilationPixels = 0.f) const;
+
+private:
+	FTimeSliderArgs TimeSliderArgs;
+
+	/** Pointer back to the timeline */
+	TWeakPtr<SGeometryCacheTimeline> WeakTimeline;
+
+	/** Brush for drawing the fill area on the scrubber */
+	const FSlateBrush* ScrubFillBrush;
+
+	/** Brush for drawing an upwards facing scrub handles */
+	const FSlateBrush* ScrubHandleUpBrush;
+
+	/** Brush for drawing a downwards facing scrub handle */
+	const FSlateBrush* ScrubHandleDownBrush;
+
+	/** Brush for drawing an editable time */
+	const FSlateBrush* EditableTimeBrush;
+
+	// Mouse Controls
+	/** Total mouse delta during dragging **/
+	float DistanceDragged;
+
+	/** If we are dragging a scrubber or dragging to set the time range */
+	enum DragType
+	{
+		DRAG_SCRUBBING_TIME,
+		DRAG_SETTING_RANGE,
+		DRAG_PLAYBACK_START,
+		DRAG_PLAYBACK_END,
+		DRAG_SELECTION_START,
+		DRAG_SELECTION_END,
+		DRAG_TIME,
+		DRAG_NONE
+	};
+
+	DragType MouseDragType;
+
+	/** If we are currently panning the panel */
+	bool bPanning;
+
+	/** Index of the current dragged time */
+	int32 DraggedTimeIndex;
+
+	/** Mouse down position range */
+	FVector2D MouseDownPosition[2];
+
+	/** Geometry on mouse down */
+	FGeometry MouseDownGeometry;
+
+	/** Range stack */
+	TArray<TRange<double>> ViewRangeStack;
+
+};
\ No newline at end of file
diff --git a/Engine/Plugins/Runtime/GeometryCache/Source/GeometryCacheEd/Private/GeometryCacheTimelineBindingAsset.cpp b/Engine/Plugins/Runtime/GeometryCache/Source/GeometryCacheEd/Private/GeometryCacheTimelineBindingAsset.cpp
new file mode 100644
--- /dev/null
+++ b/Engine/Plugins/Runtime/GeometryCache/Source/GeometryCacheEd/Private/GeometryCacheTimelineBindingAsset.cpp
@@ -0,0 +1,121 @@
+// Copyright Epic Games, Inc. All Rights Reserved.
+
+#include "GeometryCacheTimelineBindingAsset.h"
+
+#include "CommonFrameRates.h"
+#include "Fonts/FontMeasure.h"
+#include "GeometryCacheComponent.h"
+#include "GeometryCacheHelpers.h"
+#include "Styling/AppStyle.h"
+
+FGeometryCacheTimelineBindingAsset::FGeometryCacheTimelineBindingAsset(TWeakObjectPtr<UGeometryCacheComponent> InPreviewComponent)
+	:PreviewComponent(InPreviewComponent)
+{
+	if (InPreviewComponent.IsValid())
+	{
+		PlaybackRange = FAnimatedRange(0.0, (double)InPreviewComponent->GetDuration());
+		SetViewRange(PlaybackRange);
+	}
+}
+
+FFrameRate FGeometryCacheTimelineBindingAsset::GetFrameRate() const
+{
+	return FFrameRate(FMath::RoundToInt32((float)PreviewComponent.Pin()->GetNumberOfFrames() / PreviewComponent.Pin()->GetDuration()), 1);
+}
+
+int32 FGeometryCacheTimelineBindingAsset::GetTickResolution() const
+{
+	return int32(1000);
+}
+
+FAnimatedRange FGeometryCacheTimelineBindingAsset::GetViewRange() const
+{
+	return ViewRange;
+}
+
+void FGeometryCacheTimelineBindingAsset::SetViewRange(TRange<double> InRange)
+{
+	ViewRange = InRange;
+
+	if (WorkingRange.HasLowerBound() && WorkingRange.HasUpperBound())
+	{
+		WorkingRange = TRange<double>::Hull(WorkingRange, ViewRange);
+	}
+	else
+	{
+		WorkingRange = ViewRange;
+	}
+}
+
+FAnimatedRange FGeometryCacheTimelineBindingAsset::GetWorkingRange() const
+{
+	return WorkingRange;
+}
+
+TRange<FFrameNumber> FGeometryCacheTimelineBindingAsset::GetPlaybackRange() const
+{
+	const int32 Resolution = GetTickResolution();
+	return TRange<FFrameNumber>(FFrameNumber(FMath::RoundToInt32(PlaybackRange.GetLowerBoundValue() * (double)Resolution)), FFrameNumber(FMath::RoundToInt32(PlaybackRange.GetUpperBoundValue() * (double)Resolution)));
+}
+
+FFrameNumber FGeometryCacheTimelineBindingAsset::GetScrubPosition() const
+{
+	if (PreviewComponent.IsValid())
+	{
+		float SampleTime = 0.0f;
+		if(PreviewComponent->IsPlaying() && PreviewComponent->IsLooping())
+		{
+			SampleTime = GeometyCacheHelpers::WrapAnimationTime(PreviewComponent->GetElapsedTime(), PreviewComponent->GetDuration());
+		}
+		else
+		{
+			SampleTime = PreviewComponent->GetElapsedTime();
+			while (SampleTime > PreviewComponent->GetDuration())
+			{
+				SampleTime -= PreviewComponent->GetDuration();
+			}
+			while (SampleTime < 0)
+			{
+				SampleTime += PreviewComponent->GetDuration();
+			}
+			SampleTime = FMath::Clamp(SampleTime, 0.0f, PreviewComponent->GetDuration());
+		}
+
+		return FFrameNumber(FMath::RoundToInt32(SampleTime * (double)GetTickResolution()));
+	}
+
+	return FFrameNumber(0);
+}
+
+float FGeometryCacheTimelineBindingAsset::GetScrubTime() const
+{
+	if (PreviewComponent.IsValid())
+	{
+		return GeometyCacheHelpers::WrapAnimationTime(PreviewComponent->GetElapsedTime(), PreviewComponent->GetDuration());
+	}
+
+	return 0.0f;
+}
+
+void FGeometryCacheTimelineBindingAsset::SetScrubPosition(FFrameTime NewScrubPosition) const
+{
+	if (PreviewComponent.IsValid())
+	{
+		if (PreviewComponent->IsPlaying())
+		{
+			PreviewComponent->Stop();
+		}
+
+		PreviewComponent->SetCurrentTime(static_cast<float>(NewScrubPosition.AsDecimal() / static_cast<double>(GetTickResolution())));
+	}
+}
+
+void FGeometryCacheTimelineBindingAsset::HandleViewRangeChanged(TRange<double> InRange, EViewRangeInterpolation InInterpolation)
+{
+	SetViewRange(InRange);
+}
+
+void FGeometryCacheTimelineBindingAsset::HandleWorkingRangeChanged(TRange<double> InRange)
+{
+	WorkingRange = InRange;
+}
\ No newline at end of file
diff --git a/Engine/Plugins/Runtime/GeometryCache/Source/GeometryCacheEd/Private/GeometryCacheTimelineBindingAsset.h b/Engine/Plugins/Runtime/GeometryCache/Source/GeometryCacheEd/Private/GeometryCacheTimelineBindingAsset.h
new file mode 100644
--- /dev/null
+++ b/Engine/Plugins/Runtime/GeometryCache/Source/GeometryCacheEd/Private/GeometryCacheTimelineBindingAsset.h
@@ -0,0 +1,63 @@
+// Copyright Epic Games, Inc. All Rights Reserved.
+
+#pragma once
+
+#include "CoreMinimal.h"
+#include "Widgets/SWidget.h"
+#include "ITimeSlider.h"
+#include "TimeSliderArgs.h"
+#include "Input/CursorReply.h"
+#include "Input/Reply.h"
+
+class UGeometryCacheComponent;
+
+class FGeometryCacheTimelineBindingAsset : public TSharedFromThis<FGeometryCacheTimelineBindingAsset>
+{
+public:
+	FGeometryCacheTimelineBindingAsset(TWeakObjectPtr<UGeometryCacheComponent> InPreviewComponent);
+	
+	TWeakObjectPtr<UGeometryCacheComponent> GetPreviewComponent() const { return PreviewComponent; }
+
+	/** Get the framerate specified by the anim sequence */
+	FFrameRate GetFrameRate() const;
+
+	/** Get the tick resolution we are displaying at */
+	int32 GetTickResolution() const;
+
+	/** Get the current view range */
+	FAnimatedRange GetViewRange() const;
+
+	/** Set the current view range */
+	void SetViewRange(TRange<double> InRange);
+
+	/** Get the working range of the model's data */
+	FAnimatedRange GetWorkingRange() const;
+
+	/** Get the playback range of the model's data */
+	TRange<FFrameNumber> GetPlaybackRange() const;
+
+	/** Get the current scrub position */
+	FFrameNumber GetScrubPosition() const;
+
+	/** Get the current scrub time */
+	float GetScrubTime() const;
+
+	/** Set the current scrub position */
+	void SetScrubPosition(FFrameTime NewScrubPostion) const;
+
+	/** Handle the view range being changed */
+	void HandleViewRangeChanged(TRange<double> InRange, EViewRangeInterpolation InInterpolation);
+
+	/** Handle the working range being changed */
+	void HandleWorkingRangeChanged(TRange<double> InRange);
+
+private:
+	TWeakObjectPtr<UGeometryCacheComponent> PreviewComponent;
+
+	FAnimatedRange ViewRange;
+
+	FAnimatedRange WorkingRange;
+
+	FAnimatedRange PlaybackRange;
+
+};
\ No newline at end of file
diff --git a/Engine/Plugins/Runtime/GeometryCache/Source/GeometryCacheEd/Private/Widgets/SGeometryCacheEditorViewport.cpp b/Engine/Plugins/Runtime/GeometryCache/Source/GeometryCacheEd/Private/Widgets/SGeometryCacheEditorViewport.cpp
new file mode 100644
--- /dev/null
+++ b/Engine/Plugins/Runtime/GeometryCache/Source/GeometryCacheEd/Private/Widgets/SGeometryCacheEditorViewport.cpp
@@ -0,0 +1,334 @@
+// Copyright Epic Games, Inc. All Rights Reserved.
+
+#include "SGeometryCacheEditorViewport.h"
+
+#include "AdvancedPreviewScene.h"
+#include "AssetViewerSettings.h"
+#include "Editor/EditorPerProjectUserSettings.h"
+#include "Editor/UnrealEdEngine.h"
+#include "EditorViewportClient.h"
+#include "EditorViewportCommands.h"
+#include "FinalPostProcessSettings.h"
+#include "GeometryCache.h"
+#include "GeometryCacheComponent.h"
+#include "GeometryCacheMeshData.h"
+#include "SceneView.h"
+#include "SCommonEditorViewportToolbarBase.h"
+#include "Slate/SceneViewport.h"
+#include "ThumbnailRendering/ThumbnailManager.h"
+#include "UnrealEdGlobals.h"
+#include "Widgets/Docking/SDockTab.h"
+#include "Widgets/SNullWidget.h"
+#include "Widgets/Text/SRichTextBlock.h"
+
+#define LOCTEXT_NAMESPACE "SGeometryCacheEditorViewport"
+
+////// VIEWPORT CLIENT	////////////////
+
+class FGeometryCacheEditorViewportClient : public FEditorViewportClient
+{
+public:
+	FGeometryCacheEditorViewportClient(FAdvancedPreviewScene& InPreviewScene, const TSharedRef<SGeometryCacheEditorViewport>& InGeometryCacheEditorViewport);
+	virtual ~FGeometryCacheEditorViewportClient() override;
+
+	virtual FLinearColor GetBackgroundColor() const override;
+	virtual void Tick(float DeltaSeconds) override;
+	virtual void Draw(FViewport* Viewport, FCanvas* Canvas) override;
+	virtual bool ShouldOrbitCamera() const override { return true; }
+
+	virtual FSceneView* CalcSceneView(FSceneViewFamily* ViewFamily, const int32 StereoViewIndex = INDEX_NONE) override;
+
+	virtual bool CanSetWidgetMode(UE::Widget::EWidgetMode NewMode) const override { return false; }
+	virtual bool CanCycleWidgetMode() const override { return false; }
+
+	void OnAssetViewerSettingsChanged(const FName& InPropertyName);
+	void SetAdvancedShowFlagsForScene(const bool bAdvancedShowFlags);
+
+private:
+	TWeakPtr<SGeometryCacheEditorViewport> GeometryCacheEditorViewportPtr;
+};
+
+FGeometryCacheEditorViewportClient::FGeometryCacheEditorViewportClient(FAdvancedPreviewScene& InPreviewScene, const TSharedRef<SGeometryCacheEditorViewport>& InGeometryCacheEditorViewport)
+	: FEditorViewportClient(nullptr, &InPreviewScene, StaticCastSharedRef<SEditorViewport>(InGeometryCacheEditorViewport))
+{
+	GeometryCacheEditorViewportPtr = InGeometryCacheEditorViewport;
+
+	// Setup defaults for the common draw helper.
+	DrawHelper.bDrawPivot = false;
+	DrawHelper.bDrawWorldBox = false;
+	DrawHelper.bDrawKillZ = false;
+	DrawHelper.bDrawGrid = false;
+	DrawHelper.GridColorAxis = FColor(80, 80, 80);
+	DrawHelper.GridColorMajor = FColor(72, 72, 72);
+	DrawHelper.GridColorMinor = FColor(64, 64, 64);
+	DrawHelper.PerspectiveGridSize = UE_OLD_HALF_WORLD_MAX1;
+	ShowWidget(false);
+
+	SetViewMode(VMI_Lit);
+
+	//EngineShowFlags.DisableAdvancedFeatures();
+	EngineShowFlags.SetSnap(0);
+	EngineShowFlags.SetTemporalAA(true);
+	EngineShowFlags.SetShaderPrint(true);
+	
+	OverrideNearClipPlane(0.001f);
+
+	//This seems to be needed to get the correct world time in the preview.
+	SetIsSimulateInEditorViewport(true);
+
+	// Register delegate to update the show flags when the post-processing is turned on or off
+	UAssetViewerSettings::Get()->OnAssetViewerSettingsChanged().AddRaw(this, &FGeometryCacheEditorViewportClient::OnAssetViewerSettingsChanged);
+
+	UEditorPerProjectUserSettings* PerProjectSettings = GetMutableDefault<UEditorPerProjectUserSettings>();
+	UAssetViewerSettings* DefaultSettings = UAssetViewerSettings::Get();
+	PerProjectSettings->AssetViewerProfileIndex = DefaultSettings->Profiles.IsValidIndex(PerProjectSettings->AssetViewerProfileIndex) ? PerProjectSettings->AssetViewerProfileIndex : 0;
+
+	FAdvancedPreviewScene* PreviewSceneCasted = static_cast<FAdvancedPreviewScene*>(PreviewScene);
+	PreviewSceneCasted->SetProfileIndex(PerProjectSettings->AssetViewerProfileIndex);
+
+	// Set correct flags according to current profile settings
+	SetAdvancedShowFlagsForScene(UAssetViewerSettings::Get()->Profiles[PerProjectSettings->AssetViewerProfileIndex].bPostProcessingEnabled);
+}
+
+FGeometryCacheEditorViewportClient::~FGeometryCacheEditorViewportClient()
+{
+	UAssetViewerSettings::Get()->OnAssetViewerSettingsChanged().RemoveAll(this);
+}
+
+FLinearColor FGeometryCacheEditorViewportClient::GetBackgroundColor() const
+{
+	FLinearColor BackgroundColor = FLinearColor::Black;
+	return BackgroundColor;
+}
+
+FSceneView* FGeometryCacheEditorViewportClient::CalcSceneView(FSceneViewFamily* ViewFamily, const int32 StereoViewIndex)
+{
+	FSceneView* SceneView = FEditorViewportClient::CalcSceneView(ViewFamily);
+	FFinalPostProcessSettings::FCubemapEntry& CubemapEntry = *new(SceneView->FinalPostProcessSettings.ContributingCubemaps) FFinalPostProcessSettings::FCubemapEntry;
+	CubemapEntry.AmbientCubemap = GUnrealEd->GetThumbnailManager()->AmbientCubemap;
+	CubemapEntry.AmbientCubemapTintMulScaleValue = FLinearColor::White;
+	return SceneView;
+}
+
+void FGeometryCacheEditorViewportClient::Tick(float DeltaSeconds)
+{
+	FEditorViewportClient::Tick(DeltaSeconds);
+
+	// Tick the preview scene world.
+	PreviewScene->GetWorld()->Tick(LEVELTICK_All, DeltaSeconds);
+}
+
+void FGeometryCacheEditorViewportClient::Draw(FViewport* InViewport, FCanvas* Canvas)
+{
+	FEditorViewportClient::Draw(InViewport, Canvas);
+}
+
+void FGeometryCacheEditorViewportClient::OnAssetViewerSettingsChanged(const FName& InPropertyName)
+{
+	if (InPropertyName == GET_MEMBER_NAME_CHECKED(FPreviewSceneProfile, bPostProcessingEnabled) || InPropertyName == NAME_None)
+	{
+		FAdvancedPreviewScene* PreviewSceneCasted = static_cast<FAdvancedPreviewScene*>(PreviewScene);
+		UAssetViewerSettings* Settings = UAssetViewerSettings::Get();
+		const int32 ProfileIndex = PreviewSceneCasted->GetCurrentProfileIndex();
+		if (Settings->Profiles.IsValidIndex(ProfileIndex))
+		{
+			SetAdvancedShowFlagsForScene(Settings->Profiles[ProfileIndex].bPostProcessingEnabled);
+		}
+	}
+}
+
+void FGeometryCacheEditorViewportClient::SetAdvancedShowFlagsForScene(const bool bAdvancedShowFlags)
+{
+	if (bAdvancedShowFlags)
+	{
+		EngineShowFlags.EnableAdvancedFeatures();
+	}
+	else
+	{
+		EngineShowFlags.DisableAdvancedFeatures();
+	}
+}
+
+////// EDITOR			////////////////
+
+void SGeometryCacheEditorViewport::Construct(const FArguments& InArgs)
+{
+	PreviewGeometryCacheComponent = nullptr;
+
+	AdvancedPreviewScene = MakeShareable(new FAdvancedPreviewScene(FPreviewScene::ConstructionValues()));
+	AdvancedPreviewScene->SetFloorVisibility(false);
+
+	SEditorViewport::Construct(SEditorViewport::FArguments());
+}
+
+SGeometryCacheEditorViewport::~SGeometryCacheEditorViewport()
+{
+	if (SystemViewportClient.IsValid())
+	{
+		SystemViewportClient->Viewport = nullptr;
+	}
+}
+
+void SGeometryCacheEditorViewport::AddReferencedObjects(FReferenceCollector& Collector)
+{
+	if (PreviewGeometryCacheComponent != nullptr)
+	{
+		Collector.AddReferencedObject(PreviewGeometryCacheComponent);
+	}
+}
+
+void SGeometryCacheEditorViewport::Tick(const FGeometry& AllottedGeometry, const double InCurrentTime, const float InDeltaTime)
+{
+	SEditorViewport::Tick(AllottedGeometry, InCurrentTime, InDeltaTime);
+}
+
+TSharedRef<class SEditorViewport> SGeometryCacheEditorViewport::GetViewportWidget()
+{
+	return SharedThis(this);
+}
+
+TSharedPtr<FExtender> SGeometryCacheEditorViewport::GetExtenders() const
+{
+	TSharedPtr<FExtender> Result(MakeShareable(new FExtender));
+	return Result;
+}
+
+void SGeometryCacheEditorViewport::OnFloatingButtonClicked()
+{
+}
+
+void SGeometryCacheEditorViewport::SetGeometryCacheComponent(UGeometryCacheComponent* InGeometryCacheComponent)
+{
+	if (PreviewGeometryCacheComponent != nullptr)
+	{
+		AdvancedPreviewScene->RemoveComponent(PreviewGeometryCacheComponent);
+	}
+
+	PreviewGeometryCacheComponent = InGeometryCacheComponent;
+
+	if (PreviewGeometryCacheComponent != nullptr)
+	{
+		AdvancedPreviewScene->AddComponent(PreviewGeometryCacheComponent, PreviewGeometryCacheComponent->GetRelativeTransform());
+	}
+
+	if (PreviewGeometryCacheComponent != nullptr && SystemViewportClient)
+	{
+		SystemViewportClient->FocusViewportOnBox(PreviewGeometryCacheComponent->Bounds.GetBox());
+	}
+
+	SceneViewport->Invalidate();
+
+}
+
+TSharedRef<FEditorViewportClient> SGeometryCacheEditorViewport::MakeEditorViewportClient()
+{
+	SystemViewportClient = MakeShareable(new FGeometryCacheEditorViewportClient(*AdvancedPreviewScene.Get(), SharedThis(this)));
+
+	SystemViewportClient->SetViewLocation(FVector::ZeroVector);
+	SystemViewportClient->SetViewRotation(FRotator::ZeroRotator);
+	SystemViewportClient->SetViewLocationForOrbiting(FVector::ZeroVector);
+	SystemViewportClient->bSetListenerPosition = false;
+
+	SystemViewportClient->SetRealtime(true);
+	SystemViewportClient->VisibilityDelegate.BindSP(this, &SGeometryCacheEditorViewport::IsVisible);
+
+	return SystemViewportClient.ToSharedRef();
+}
+
+TSharedPtr<SWidget> SGeometryCacheEditorViewport::MakeViewportToolbar()
+{
+	TAttribute<FText> StatsTextAttribute;
+	{
+		StatsTextAttribute.Bind(TAttribute<FText>::FGetter::CreateSP(this, &SGeometryCacheEditorViewport::BuildStatsText));
+	}
+
+	return SNew(SVerticalBox)
+		+ SVerticalBox::Slot()
+		.AutoHeight()
+		[
+			SNew(SBorder)
+				.BorderImage(FAppStyle::Get().GetBrush("EditorViewportToolBar.Background"))
+				.Cursor(EMouseCursor::Default)
+				[
+					SNew(SCommonEditorViewportToolbarBase, SharedThis(this))
+				]
+		]
+		+ SVerticalBox::Slot()
+		.Padding(FMargin(4.0f, 3.0f, 0.0f, 0.0f))
+		[
+			SNew(SRichTextBlock)
+				.Visibility(EVisibility::SelfHitTestInvisible)
+				.DecoratorStyleSet(&FAppStyle::Get())
+				.Text(StatsTextAttribute)
+				.TextStyle(&FAppStyle::Get().GetWidgetStyle<FTextBlockStyle>("AnimViewport.MessageText"))
+		];
+}
+
+EVisibility SGeometryCacheEditorViewport::OnGetViewportContentVisibility() const
+{
+	EVisibility BaseVisibility = SEditorViewport::OnGetViewportContentVisibility();
+	if (BaseVisibility != EVisibility::Visible)
+	{
+		return BaseVisibility;
+	}
+	return IsVisible() ? EVisibility::Visible : EVisibility::Collapsed;
+}
+
+void SGeometryCacheEditorViewport::BindCommands()
+{
+	SEditorViewport::BindCommands();
+}
+
+void SGeometryCacheEditorViewport::OnFocusViewportToSelection()
+{
+	if (PreviewGeometryCacheComponent)
+	{
+		SystemViewportClient->FocusViewportOnBox(PreviewGeometryCacheComponent->Bounds.GetBox());
+	}
+}
+
+void SGeometryCacheEditorViewport::PopulateViewportOverlays(TSharedRef<class SOverlay> Overlay)
+{
+}
+
+bool SGeometryCacheEditorViewport::IsVisible() const
+{
+	return ViewportWidget.IsValid() && (!ParentTab.IsValid() || ParentTab.Pin()->IsForeground()) && SEditorViewport::IsVisible();
+}
+
+FText SGeometryCacheEditorViewport::BuildStatsText() const
+{
+	if (!PreviewGeometryCacheComponent)
+	{
+		return FText::FromString("No component to preview");
+	}
+		
+	const FText CompName = FText::FromString(PreviewGeometryCacheComponent->GetName());
+	const int32 Tracks = PreviewGeometryCacheComponent->GetNumberOfTracks();
+	const int32 Frames = PreviewGeometryCacheComponent->GetNumberOfFrames();
+	const float Duration = PreviewGeometryCacheComponent->GetDuration();
+	const FText DurationText = FText::FromString(FString::Printf(TEXT("%.2f s"), Duration));
+	const int32 FPS = FMath::FloorToInt((float)Frames / Duration);
+	int32 NumTriangles = 0;
+
+	if (UGeometryCache* GeometryCache = PreviewGeometryCacheComponent->GetGeometryCache())
+	{
+		constexpr float Time = 0.0f;
+		TArray<FGeometryCacheMeshData> MeshData;
+		GeometryCache->GetMeshDataAtTime(Time, MeshData);
+		if (MeshData.Num() > 0)
+		{
+			NumTriangles = MeshData[0].Positions.Num() / 3;
+		}
+	}
+
+	return FText::FormatNamed(LOCTEXT("StatsText", "Previewing {CompName}\nTracks: {Tracks}\nFrames: {Frames}\nDuration:{Duration}\nFPS: {FPS} fps\nTriangles: {NumTriangles}"),
+		TEXT("CompName"), CompName,
+		TEXT("Tracks"), Tracks,
+		TEXT("Frames"), Frames,
+		TEXT("Duration"), DurationText,
+		TEXT("FPS"), FPS,
+		TEXT("NumTriangles"), NumTriangles);
+}
+
+#undef LOCTEXT_NAMESPACE
\ No newline at end of file
diff --git a/Engine/Plugins/Runtime/GeometryCache/Source/GeometryCacheEd/Private/Widgets/SGeometryCacheEditorViewport.h b/Engine/Plugins/Runtime/GeometryCache/Source/GeometryCacheEd/Private/Widgets/SGeometryCacheEditorViewport.h
new file mode 100644
--- /dev/null
+++ b/Engine/Plugins/Runtime/GeometryCache/Source/GeometryCacheEd/Private/Widgets/SGeometryCacheEditorViewport.h
@@ -0,0 +1,70 @@
+// Copyright Epic Games, Inc. All Rights Reserved.
+
+#pragma once
+
+#include "CoreMinimal.h"
+#include "SlateFwd.h"
+#include "UObject/GCObject.h"
+#include "Layout/Visibility.h"
+#include "Widgets/DeclarativeSyntaxSupport.h"
+#include "PreviewScene.h"
+#include "Framework/Commands/UICommandList.h"
+#include "EditorViewportClient.h"
+#include "SEditorViewport.h"
+#include "SCommonEditorViewportToolbarBase.h"
+#include "HairStrandsInterface.h"
+
+class FGeometryCacheEditorViewportClient;
+class FAdvancedPreviewScene;
+class UGeometryCacheComponent;
+
+class SGeometryCacheEditorViewport : public SEditorViewport, public FGCObject, public ICommonEditorViewportToolbarInfoProvider
+{
+public:
+	void Construct(const FArguments& InArgs);
+	~SGeometryCacheEditorViewport();
+
+	virtual void AddReferencedObjects(FReferenceCollector& Collector) override;
+	virtual FString GetReferencerName() const override
+	{
+		return TEXT("SGeometryCacheEditorViewport");
+	}
+
+	virtual void Tick(const FGeometry& AllottedGeometry, const double InCurrentTime, const float InDeltaTime) override;
+
+	// ICommonEditorViewportToolbarInfoProvider interface
+	virtual TSharedRef<SEditorViewport> GetViewportWidget() override;
+	virtual TSharedPtr<FExtender> GetExtenders() const override;
+	virtual void OnFloatingButtonClicked() override;
+	// End of ICommonEditorViewportToolbarInfoProvider interface
+
+	// Set the component to preview
+	void SetGeometryCacheComponent(UGeometryCacheComponent* InGeometryCacheComponent);
+
+	TSharedPtr<FAdvancedPreviewScene> GetAdvancedPreviewScene() { return AdvancedPreviewScene; }
+	
+protected:
+	/** SEditorViewport interface */
+	virtual TSharedRef<FEditorViewportClient> MakeEditorViewportClient() override;
+	virtual TSharedPtr<SWidget> MakeViewportToolbar() override;
+	virtual EVisibility OnGetViewportContentVisibility() const override;
+	virtual void BindCommands() override;
+	virtual void OnFocusViewportToSelection() override;
+	virtual void PopulateViewportOverlays(TSharedRef<class SOverlay> Overlay) override;
+
+private:
+	bool IsVisible() const override;
+
+	FText BuildStatsText() const;
+private:
+	/** The parent tab where this viewport resides */
+	TWeakPtr<SDockTab> ParentTab;
+
+	/** Level viewport client */
+	TSharedPtr<FGeometryCacheEditorViewportClient> SystemViewportClient;
+
+	/** Preview Scene - uses advanced preview settings */
+	TSharedPtr<FAdvancedPreviewScene> AdvancedPreviewScene;
+
+	TObjectPtr<UGeometryCacheComponent> PreviewGeometryCacheComponent;
+};
\ No newline at end of file
diff --git a/Engine/Plugins/Runtime/GeometryCache/Source/GeometryCacheEd/Private/Widgets/SGeometryCacheTimeline.cpp b/Engine/Plugins/Runtime/GeometryCache/Source/GeometryCacheEd/Private/Widgets/SGeometryCacheTimeline.cpp
new file mode 100644
--- /dev/null
+++ b/Engine/Plugins/Runtime/GeometryCache/Source/GeometryCacheEd/Private/Widgets/SGeometryCacheTimeline.cpp
@@ -0,0 +1,401 @@
+// Copyright Epic Games, Inc. All Rights Reserved.
+
+#include "SGeometryCacheTimeline.h"
+
+#include "Fonts/FontMeasure.h"
+#include "FrameNumberNumericInterface.h"
+#include "FrameNumberTimeEvaluator.h"
+#include "Framework/Application/SlateApplication.h"
+#include "GeometryCache.h"
+#include "GeometryCacheComponent.h"
+#include "GeometryCacheHelpers.h"
+#include "GeometryCacheTimelineBindingAsset.h"
+#include "GeometryCacheTimeSliderController.h"
+#include "ISequencerWidgetsModule.h"
+#include "Modules/ModuleManager.h"
+#include "Widgets/Layout/SGridPanel.h"
+#include "Widgets/Layout/SScrollBorder.h"
+#include "Widgets/Layout/SScrollBox.h"
+#include "Widgets/Layout/SSpacer.h"
+#include "Widgets/SBoxPanel.h"
+#include "Widgets/SGeometryCacheTimelineOverlay.h"
+#include "Widgets/SGeometryCacheTimelineSplitterOverlay.h"
+#include "Widgets/SGeometryCacheTimelineTransportControls.h"
+#include "Widgets/Text/STextBlock.h"
+#include "Widgets/Views/SListView.h"
+
+#define LOCTEXT_NAMESPACE "SGeometryCacheTimeline"
+
+void SGeometryCacheTimeline::Construct(const FArguments& InArgs, const TSharedRef<FGeometryCacheTimelineBindingAsset>& InBindingAsset)
+{
+	TWeakPtr<FGeometryCacheTimelineBindingAsset>  WeakBindingAsset = InBindingAsset;
+	BindingAsset = InBindingAsset;
+
+	const EFrameNumberDisplayFormats DisplayFormat = EFrameNumberDisplayFormats::Frames;
+	const TAttribute<FFrameRate> TickResolution = MakeAttributeLambda([WeakBindingAsset]()
+		{
+			return FFrameRate(WeakBindingAsset.Pin()->GetTickResolution(), 1);
+		});
+
+	ViewRange = MakeAttributeLambda([WeakBindingAsset]() { return WeakBindingAsset.IsValid() ? WeakBindingAsset.Pin()->GetViewRange() : FAnimatedRange(0.0, 0.0); });
+	
+	const TAttribute<FFrameRate> DisplayRate = MakeAttributeLambda([WeakBindingAsset]()
+		{
+			return WeakBindingAsset.Pin()->GetFrameRate();
+		});
+
+	ColumnFillCoefficients[0] = 0.3f;
+	ColumnFillCoefficients[1] = 0.7f;
+
+	TAttribute<float> FillCoefficient_0, FillCoefficient_1;
+	{
+		FillCoefficient_0.Bind(TAttribute<float>::FGetter::CreateSP(this, &SGeometryCacheTimeline::GetColumnFillCoefficient, 0));
+		FillCoefficient_1.Bind(TAttribute<float>::FGetter::CreateSP(this, &SGeometryCacheTimeline::GetColumnFillCoefficient, 1));
+	};
+
+	NumericTypeInterface = MakeShareable(new FFrameNumberInterface(DisplayFormat, 0, TickResolution, DisplayRate));
+
+	FTimeSliderArgs TimeSliderArgs;
+	{
+		TimeSliderArgs.ScrubPosition = MakeAttributeLambda([WeakBindingAsset]() { return WeakBindingAsset.IsValid() ? WeakBindingAsset.Pin()->GetScrubPosition() : FFrameTime(0); });
+		TimeSliderArgs.ViewRange = ViewRange;
+		TimeSliderArgs.PlaybackRange = MakeAttributeLambda([WeakBindingAsset]() { return WeakBindingAsset.IsValid() ? WeakBindingAsset.Pin()->GetPlaybackRange() : TRange<FFrameNumber>(0, 0); });
+		TimeSliderArgs.ClampRange = MakeAttributeLambda([WeakBindingAsset]() { return WeakBindingAsset.IsValid() ? WeakBindingAsset.Pin()->GetWorkingRange() : FAnimatedRange(0.0, 0.0); });
+		TimeSliderArgs.DisplayRate = DisplayRate;
+		TimeSliderArgs.TickResolution = TickResolution;
+		TimeSliderArgs.OnViewRangeChanged = FOnViewRangeChanged::CreateSP(&InBindingAsset.Get(), &FGeometryCacheTimelineBindingAsset::HandleViewRangeChanged);
+		TimeSliderArgs.OnClampRangeChanged = FOnTimeRangeChanged::CreateSP(&InBindingAsset.Get(), &FGeometryCacheTimelineBindingAsset::HandleWorkingRangeChanged);
+		TimeSliderArgs.IsPlaybackRangeLocked = true;
+		TimeSliderArgs.PlaybackStatus = EMovieScenePlayerStatus::Stopped;
+		TimeSliderArgs.NumericTypeInterface = NumericTypeInterface;
+		TimeSliderArgs.OnScrubPositionChanged = FOnScrubPositionChanged::CreateSP(this, &SGeometryCacheTimeline::HandleScrubPositionChanged);
+	}
+
+	TimeSliderController = MakeShareable(new FGeometryCacheTimeSlideController(TimeSliderArgs, SharedThis(this)));
+
+	TSharedRef<FGeometryCacheTimeSlideController> TimeSliderControllerRef = TimeSliderController.ToSharedRef();
+
+	// Create the top slider
+	constexpr bool bMirrorLabels = false;
+	ISequencerWidgetsModule& SequencerWidgets = FModuleManager::Get().LoadModuleChecked<ISequencerWidgetsModule>("SequencerWidgets");
+	TimeSlider = SequencerWidgets.CreateTimeSlider(TimeSliderControllerRef, bMirrorLabels);
+
+	// Create bottom time range slider
+	TSharedRef<ITimeSlider> BottomTimeRange = SequencerWidgets.CreateTimeRange(
+		FTimeRangeArgs(
+			EShowRange::ViewRange | EShowRange::WorkingRange | EShowRange::PlaybackRange,
+			EShowRange::ViewRange | EShowRange::WorkingRange,
+			TimeSliderControllerRef,
+			EVisibility::Visible,
+			NumericTypeInterface.ToSharedRef()
+		),
+		SequencerWidgets.CreateTimeRangeSlider(TimeSliderControllerRef)
+	);
+
+	TSharedRef<SScrollBar> ScrollBar = SNew(SScrollBar)
+		.Thickness(FVector2D(5.0f, 5.0f));
+
+	InitTrackNames();
+	
+	TracksListView = SNew(SListView< TSharedPtr<FString> >)
+		.ExternalScrollbar(ScrollBar)
+		.ListItemsSource(&TrackNames)
+		.SelectionMode(ESelectionMode::Single)
+		.OnGenerateRow(this, &SGeometryCacheTimeline::HandleTimelineListViewGenerateRow);
+
+	// Grid Panel Constants
+	const int32 Column0 = 0, Column1 = 1;
+	const int32 Row0 = 0, Row1 = 1, Row2 = 2, Row3 = 3, Row4 = 4;
+
+	const float CommonPadding = 3.f;
+	const FMargin ResizeBarPadding(4.f, 0, 0, 0);
+
+	ChildSlot
+		[
+			SNew(SOverlay)
+				+ SOverlay::Slot()
+				[
+					SNew(SGridPanel)
+						.FillRow(1, 1.0f)
+						.FillColumn(0, FillCoefficient_0)
+						.FillColumn(1, FillCoefficient_1)
+						+ SGridPanel::Slot(Column1, Row0, SGridPanel::Layer(10))
+						.Padding(ResizeBarPadding)
+						[
+							SNew(SBorder)
+								.BorderImage(FAppStyle::GetBrush("ToolPanel.GroupBorder"))
+								.BorderBackgroundColor(FLinearColor(.50f, .50f, .50f, 1.0f))
+								.Padding(0.f)
+								.Clipping(EWidgetClipping::ClipToBounds)
+								[
+									TimeSlider.ToSharedRef()
+								]
+						]
+						+ SGridPanel::Slot(Column1, Row1, SGridPanel::Layer(10))
+						.Padding(ResizeBarPadding)
+						[
+							SNew(SGeometryCacheTimelineOverlay, TimeSliderControllerRef)
+								.Visibility(EVisibility::HitTestInvisible)
+								.DisplayScrubPosition(false)
+								.DisplayTickLines(true)
+								.Clipping(EWidgetClipping::ClipToBounds)
+								.PaintPlaybackRangeArgs(FPaintPlaybackRangeArgs(FAppStyle::GetBrush("Sequencer.Timeline.PlayRange_L"), FAppStyle::GetBrush("Sequencer.Timeline.PlayRange_R"), 6.f))
+						]
+
+						// Overlay that draws the scrub position
+						+ SGridPanel::Slot(Column1, Row1, SGridPanel::Layer(20))
+						.Padding(ResizeBarPadding)
+						[
+							SNew(SGeometryCacheTimelineOverlay, TimeSliderControllerRef)
+								.Visibility(EVisibility::HitTestInvisible)
+								.DisplayScrubPosition(true)
+								.DisplayTickLines(false)
+								.Clipping(EWidgetClipping::ClipToBounds)
+						]
+						+ SGridPanel::Slot(Column1, Row3, SGridPanel::Layer(10))
+						.Padding(ResizeBarPadding)
+						[
+							SNew(SBorder)
+								.BorderImage(FAppStyle::GetBrush("ToolPanel.GroupBorder"))
+								.BorderBackgroundColor(FLinearColor(0.5f, 0.5f, 0.5f, 1.0f))
+								.Clipping(EWidgetClipping::ClipToBounds)
+								.Padding(0)
+								[
+									BottomTimeRange
+								]
+						]
+						+ SGridPanel::Slot(Column0, Row3, SGridPanel::Layer(10))
+						.VAlign(VAlign_Center)
+						.HAlign(HAlign_Center)
+						[
+							SNew(SGeometryCacheTimelineTransportControls, InBindingAsset)
+						]
+						+ SGridPanel::Slot(Column0, Row1, SGridPanel::Layer(5))
+						.ColumnSpan(2)
+						[
+							SNew(SHorizontalBox)
+								+ SHorizontalBox::Slot()
+								[
+									SNew(SOverlay)
+										+ SOverlay::Slot()
+										[
+											SNew(SVerticalBox)
+												+ SVerticalBox::Slot()
+												.FillHeight(1.f)
+												[
+													SNew(SScrollBorder, TracksListView.ToSharedRef())
+													[
+														SNew(SHorizontalBox)
+														+SHorizontalBox::Slot()
+														.FillWidth(FillCoefficient_0)
+														[
+															SNew(SBox)
+																[
+																	TracksListView.ToSharedRef()
+																]
+														]
+														+ SHorizontalBox::Slot()
+														.FillWidth(FillCoefficient_1)
+														[
+															SNew(SBox)
+															.Padding(ResizeBarPadding)
+															.Clipping(EWidgetClipping::ClipToBounds)	
+														]
+													]
+												]
+										]
+										+ SOverlay::Slot()
+										.HAlign(HAlign_Right)
+										[
+											ScrollBar
+										]
+								]
+						]
+				]
+				+ SOverlay::Slot()
+				[
+					// track area virtual splitter overlay
+					SNew(SGeometryCacheTimelineSplitterOverlay)
+						.Style(FAppStyle::Get(), "AnimTimeline.Outliner.Splitter")
+						.Visibility(EVisibility::SelfHitTestInvisible)
+
+						+ SSplitter::Slot()
+						.Value(FillCoefficient_0)
+						.OnSlotResized(SSplitter::FOnSlotResized::CreateSP(this, &SGeometryCacheTimeline::OnColumnFillCoefficientChanged, 0))
+						[
+							SNew(SSpacer)
+						]
+
+						+ SSplitter::Slot()
+						.Value(FillCoefficient_1)
+						.OnSlotResized(SSplitter::FOnSlotResized::CreateSP(this, &SGeometryCacheTimeline::OnColumnFillCoefficientChanged, 1))
+						[
+							SNew(SSpacer)
+						]
+				]
+		];
+
+}
+
+void SGeometryCacheTimeline::HandleScrubPositionChanged(FFrameTime NewScrubPosition, bool bIsScrubbing, bool bEvaluate) const
+{
+	if (BindingAsset.IsValid())
+	{
+		TWeakObjectPtr<UGeometryCacheComponent> GeometryCacheComponent = BindingAsset.Pin()->GetPreviewComponent();
+		if (GeometryCacheComponent.IsValid() && GeometryCacheComponent->IsPlaying())
+		{
+			GeometryCacheComponent->Stop();
+		}
+
+		BindingAsset.Pin()->SetScrubPosition(NewScrubPosition);
+	}
+
+}
+
+void SGeometryCacheTimeline::InitTrackNames()
+{
+	TrackNames.Empty();
+	if (BindingAsset.IsValid())
+	{
+		const TArray<FString> Names = BindingAsset.Pin()->GetPreviewComponent()->GetTrackNames();
+		for (const FString& Name : Names)
+		{
+			TrackNames.Add(MakeShareable(new FString(Name)));
+		}
+	}
+
+	
+}
+
+void SGeometryCacheTimeline::OnColumnFillCoefficientChanged(float FillCoefficient, int32 ColumnIndex)
+{
+	ColumnFillCoefficients[ColumnIndex] = FillCoefficient;
+}
+
+TSharedRef<ITableRow> SGeometryCacheTimeline::HandleTimelineListViewGenerateRow(TSharedPtr<FString> Text, const TSharedRef<STableViewBase>& OwnerTable)
+{
+	return SNew(STableRow<TSharedPtr<FString> >, OwnerTable)
+		[
+			SNew(SBox)
+				[
+					SNew(STextBlock)
+						.Margin(FMargin{10,5})
+						.Text(FText::FromString(*Text))
+				]
+		];
+}
+
+// FFrameRate::ComputeGridSpacing doesnt deal well with prime numbers, so we have a custom impl here
+static bool ComputeGridSpacing(const FFrameRate& InFrameRate, float PixelsPerSecond, double& OutMajorInterval, int32& OutMinorDivisions, float MinTickPx, float DesiredMajorTickPx)
+{
+	// First try built-in spacing
+	const bool bResult = InFrameRate.ComputeGridSpacing(PixelsPerSecond, OutMajorInterval, OutMinorDivisions, MinTickPx, DesiredMajorTickPx);
+	if (!bResult || OutMajorInterval == 1.0)
+	{
+		if (PixelsPerSecond <= 0.f)
+		{
+			return false;
+		}
+
+		const int32 RoundedFPS = static_cast<int32>(FMath::RoundToInt(InFrameRate.AsDecimal()));
+
+		if (RoundedFPS > 0)
+		{
+			// Showing frames
+			TArray<int32, TInlineAllocator<10>> CommonBases;
+
+			// Divide the rounded frame rate by 2s, 3s or 5s recursively
+			{
+				const int32 Denominators[] = { 2, 3, 5 };
+
+				int32 LowestBase = RoundedFPS;
+				for (;;)
+				{
+					CommonBases.Add(LowestBase);
+
+					if (LowestBase % 2 == 0) { LowestBase = LowestBase / 2; }
+					else if (LowestBase % 3 == 0) { LowestBase = LowestBase / 3; }
+					else if (LowestBase % 5 == 0) { LowestBase = LowestBase / 5; }
+					else
+					{
+						int32 LowestResult = LowestBase;
+						for (int32 Denominator : Denominators)
+						{
+							int32 Result = LowestBase / Denominator;
+							if (Result > 0 && Result < LowestResult)
+							{
+								LowestResult = Result;
+							}
+						}
+
+						if (LowestResult < LowestBase)
+						{
+							LowestBase = LowestResult;
+						}
+						else
+						{
+							break;
+						}
+					}
+				}
+			}
+
+			Algo::Reverse(CommonBases);
+
+			const int32 Scale = static_cast<int32>(FMath::CeilToInt(DesiredMajorTickPx / PixelsPerSecond * InFrameRate.AsDecimal()));
+			const int32 BaseIndex = FMath::Min(Algo::LowerBound(CommonBases, Scale), CommonBases.Num() - 1);
+			const int32 Base = CommonBases[BaseIndex];
+
+			const int32 MajorIntervalFrames = FMath::CeilToInt((float)Scale / static_cast<float>(Base)) * Base;
+			OutMajorInterval = MajorIntervalFrames * InFrameRate.AsInterval();
+
+			// Find the lowest number of divisions we can show that's larger than the minimum tick size
+			OutMinorDivisions = 0;
+			for (int32 DivIndex = 0; DivIndex < BaseIndex; ++DivIndex)
+			{
+				if (Base % CommonBases[DivIndex] == 0)
+				{
+					const int32 MinorDivisions = MajorIntervalFrames / CommonBases[DivIndex];
+					if (OutMajorInterval / MinorDivisions * PixelsPerSecond >= MinTickPx)
+					{
+						OutMinorDivisions = MinorDivisions;
+						break;
+					}
+				}
+			}
+		}
+	}
+
+	return OutMajorInterval != 0;
+}
+
+bool SGeometryCacheTimeline::GetGridMetrics(float PhysicalWidth, double& OutMajorInterval, int32& OutMinorDivisions) const
+{
+	const FSlateFontInfo SmallLayoutFont = FCoreStyle::GetDefaultFontStyle("Regular", 8);
+	const TSharedRef<FSlateFontMeasure> FontMeasureService = FSlateApplication::Get().GetRenderer()->GetFontMeasureService();
+
+	const FFrameRate DisplayRate = BindingAsset.Pin()->GetFrameRate();
+	const double BiggestTime = ViewRange.Get().GetUpperBoundValue();
+	const FString TickString = NumericTypeInterface->ToString((BiggestTime * DisplayRate).FrameNumber.Value);
+	const FVector2D MaxTextSize = FontMeasureService->Measure(TickString, SmallLayoutFont);
+
+	constexpr float MajorTickMultiplier = 2.f;
+
+	const float MinTickPx = static_cast<float>(MaxTextSize.X) + 5.f;
+	const float DesiredMajorTickPx = static_cast<float>(MaxTextSize.X) * MajorTickMultiplier;
+
+	if (PhysicalWidth > 0 && DisplayRate.AsDecimal() > 0)
+	{
+		return ComputeGridSpacing(
+			DisplayRate,
+			static_cast<float>(PhysicalWidth / ViewRange.Get().Size<double>()),
+			OutMajorInterval,
+			OutMinorDivisions,
+			MinTickPx,
+			DesiredMajorTickPx);
+	}
+
+	return false;
+}
+
+#undef LOCTEXT_NAMESPACE
\ No newline at end of file
diff --git a/Engine/Plugins/Runtime/GeometryCache/Source/GeometryCacheEd/Private/Widgets/SGeometryCacheTimeline.h b/Engine/Plugins/Runtime/GeometryCache/Source/GeometryCacheEd/Private/Widgets/SGeometryCacheTimeline.h
new file mode 100644
--- /dev/null
+++ b/Engine/Plugins/Runtime/GeometryCache/Source/GeometryCacheEd/Private/Widgets/SGeometryCacheTimeline.h
@@ -0,0 +1,64 @@
+// Copyright Epic Games, Inc. All Rights Reserved.
+
+#pragma once
+
+#include "CoreMinimal.h"
+#include "ITimeSlider.h"
+#include "AnimatedRange.h"
+#include "CoreTypes.h"
+#include "Widgets/SCompoundWidget.h"
+#include "Widgets/Input/NumericTypeInterface.h"
+
+#include "Widgets/Views/STableViewBase.h"
+#include "Widgets/Views/STableRow.h"
+
+class FGeometryCacheTimeSlideController;
+class UGeometryCacheComponent;
+
+class FGeometryCacheTimelineBindingAsset;
+
+class SGeometryCacheTimeline : public SCompoundWidget
+{
+public:
+	SLATE_BEGIN_ARGS(SGeometryCacheTimeline) { }
+	SLATE_END_ARGS()
+
+	void Construct(const FArguments& InArgs, const TSharedRef<FGeometryCacheTimelineBindingAsset>& InBindingAsset);
+
+	/** Get numeric Type interface for converting between frame numbers and display formats. */
+	TSharedRef<INumericTypeInterface<double>> GetNumericTypeInterface() const { return NumericTypeInterface.ToSharedRef(); }
+
+	/** Compute a major grid interval and number of minor divisions to display */
+	bool GetGridMetrics(float PhysicalWidth, double& OutMajorInterval, int32& OutMinorDivisions) const;
+
+private:
+	void HandleScrubPositionChanged(FFrameTime NewScrubPosition, bool bIsScrubbing, bool bEvaluate) const;
+
+	void InitTrackNames();
+
+	float GetColumnFillCoefficient(int32 ColumnIndex) const
+	{
+		return ColumnFillCoefficients[ColumnIndex];
+	}
+
+	void OnColumnFillCoefficientChanged(float FillCoefficient, int32 ColumnIndex);
+
+	TSharedRef<ITableRow> HandleTimelineListViewGenerateRow(TSharedPtr<FString> Text, const TSharedRef<STableViewBase>& OwnerTable);
+private:
+	/** The fill coefficients of each column in the grid. */
+	float ColumnFillCoefficients[2];
+	
+	TWeakPtr<FGeometryCacheTimelineBindingAsset> BindingAsset;
+
+	TSharedPtr<INumericTypeInterface<double>> NumericTypeInterface;
+
+	TSharedPtr<FGeometryCacheTimeSlideController> TimeSliderController;
+
+	TSharedPtr<ITimeSlider> TimeSlider;
+
+	TArray<TSharedPtr<FString>> TrackNames;
+
+	TSharedPtr<SListView<TSharedPtr<FString> > > TracksListView;
+
+	TAttribute<FAnimatedRange> ViewRange;
+};
\ No newline at end of file
diff --git a/Engine/Plugins/Runtime/GeometryCache/Source/GeometryCacheEd/Private/Widgets/SGeometryCacheTimelineOverlay.cpp b/Engine/Plugins/Runtime/GeometryCache/Source/GeometryCacheEd/Private/Widgets/SGeometryCacheTimelineOverlay.cpp
new file mode 100644
--- /dev/null
+++ b/Engine/Plugins/Runtime/GeometryCache/Source/GeometryCacheEd/Private/Widgets/SGeometryCacheTimelineOverlay.cpp
@@ -0,0 +1,28 @@
+// Copyright Epic Games, Inc. All Rights Reserved.
+
+#include "SGeometryCacheTimelineOverlay.h"
+
+#include "GeometryCacheTimeSliderController.h"
+#include "HAL/PlatformCrt.h"
+#include "Misc/Optional.h"
+
+class FSlateRect;
+class FWidgetStyle;
+struct FGeometry;
+
+int32 SGeometryCacheTimelineOverlay::OnPaint(const FPaintArgs& Args, const FGeometry& AllottedGeometry, const FSlateRect& MyCullingRect, FSlateWindowElementList& OutDrawElements, int32 LayerId, const FWidgetStyle& InWidgetStyle, bool bParentEnabled) const
+{
+	FPaintViewAreaArgs PaintArgs;
+	PaintArgs.bDisplayTickLines = bDisplayTickLines.Get();
+	PaintArgs.bDisplayScrubPosition = bDisplayScrubPosition.Get();
+
+	if (PaintPlaybackRangeArgs.IsSet())
+	{
+		PaintArgs.PlaybackRangeArgs = PaintPlaybackRangeArgs.Get();
+	}
+
+	TimeSliderController->OnPaintViewArea(AllottedGeometry, MyCullingRect, OutDrawElements, LayerId, ShouldBeEnabled(bParentEnabled), PaintArgs);
+
+	return SCompoundWidget::OnPaint(Args, AllottedGeometry, MyCullingRect, OutDrawElements, LayerId, InWidgetStyle, bParentEnabled);
+}
+
diff --git a/Engine/Plugins/Runtime/GeometryCache/Source/GeometryCacheEd/Private/Widgets/SGeometryCacheTimelineOverlay.h b/Engine/Plugins/Runtime/GeometryCache/Source/GeometryCacheEd/Private/Widgets/SGeometryCacheTimelineOverlay.h
new file mode 100644
--- /dev/null
+++ b/Engine/Plugins/Runtime/GeometryCache/Source/GeometryCacheEd/Private/Widgets/SGeometryCacheTimelineOverlay.h
@@ -0,0 +1,54 @@
+// Copyright Epic Games, Inc. All Rights Reserved.
+
+#pragma once
+
+#include "HAL/Platform.h"
+#include "ITimeSlider.h"
+#include "Misc/Attribute.h"
+#include "Templates/SharedPointer.h"
+#include "Widgets/DeclarativeSyntaxSupport.h"
+#include "Widgets/SCompoundWidget.h"
+
+class FGeometryCacheTimeSlideController;
+class FPaintArgs;
+class FSlateRect;
+class FSlateWindowElementList;
+class FWidgetStyle;
+struct FGeometry;
+
+class SGeometryCacheTimelineOverlay : public SCompoundWidget
+{
+public:
+	SLATE_BEGIN_ARGS(SGeometryCacheTimelineOverlay)
+		: _DisplayTickLines(true)
+		, _DisplayScrubPosition(false)
+		{}
+
+		SLATE_ATTRIBUTE(bool, DisplayTickLines)
+		SLATE_ATTRIBUTE(bool, DisplayScrubPosition)
+		SLATE_ATTRIBUTE(FPaintPlaybackRangeArgs, PaintPlaybackRangeArgs)
+
+	SLATE_END_ARGS()
+
+	void Construct(const FArguments& InArgs, TSharedRef<FGeometryCacheTimeSlideController> InTimeSliderController)
+	{
+		bDisplayScrubPosition = InArgs._DisplayScrubPosition;
+		bDisplayTickLines = InArgs._DisplayTickLines;
+		PaintPlaybackRangeArgs = InArgs._PaintPlaybackRangeArgs;
+		TimeSliderController = InTimeSliderController;
+	}
+
+private:
+	/** SWidget Interface */
+	virtual int32 OnPaint(const FPaintArgs& Args, const FGeometry& AllottedGeometry, const FSlateRect& MyCullingRect, FSlateWindowElementList& OutDrawElements, int32 LayerId, const FWidgetStyle& InWidgetStyle, bool bParentEnabled) const override;
+
+private:
+	/** Controller for manipulating time */
+	TSharedPtr<FGeometryCacheTimeSlideController> TimeSliderController;
+	/** Whether or not to display the scrub position */
+	TAttribute<bool> bDisplayScrubPosition;
+	/** Whether or not to display tick lines */
+	TAttribute<bool> bDisplayTickLines;
+	/** User-supplied options for drawing playback range */
+	TAttribute<FPaintPlaybackRangeArgs> PaintPlaybackRangeArgs;
+};
\ No newline at end of file
diff --git a/Engine/Plugins/Runtime/GeometryCache/Source/GeometryCacheEd/Private/Widgets/SGeometryCacheTimelineSplitterOverlay.cpp b/Engine/Plugins/Runtime/GeometryCache/Source/GeometryCacheEd/Private/Widgets/SGeometryCacheTimelineSplitterOverlay.cpp
new file mode 100644
--- /dev/null
+++ b/Engine/Plugins/Runtime/GeometryCache/Source/GeometryCacheEd/Private/Widgets/SGeometryCacheTimelineSplitterOverlay.cpp
@@ -0,0 +1,131 @@
+// Copyright Epic Games, Inc. All Rights Reserved.
+
+#include "SGeometryCacheTimelineSplitterOverlay.h"
+
+#include "Delegates/Delegate.h"
+#include "HAL/PlatformCrt.h"
+#include "Layout/ArrangedChildren.h"
+#include "Layout/ArrangedWidget.h"
+#include "Layout/ChildrenBase.h"
+#include "Layout/Geometry.h"
+#include "Layout/Visibility.h"
+#include "Math/Vector2D.h"
+#include "Misc/Attribute.h"
+#include "SlotBase.h"
+#include "Types/SlateEnums.h"
+#include "Widgets/DeclarativeSyntaxSupport.h"
+#include "Widgets/Layout/SBox.h"
+
+struct FCaptureLostEvent;
+struct FPointerEvent;
+
+void SGeometryCacheTimelineSplitterOverlay::Construct( const FArguments& InArgs )
+{
+	SetVisibility(EVisibility::SelfHitTestInvisible);
+
+	Splitter = SArgumentNew(InArgs, SSplitter);
+	Splitter->SetVisibility(EVisibility::HitTestInvisible);
+	AddSlot()
+	[
+		Splitter.ToSharedRef()
+	];
+
+	for (int32 Index = 0; Index < Splitter->GetChildren()->Num() - 1; ++Index)
+	{
+		AddSlot()
+		.Padding(TAttribute<FMargin>::Create(TAttribute<FMargin>::FGetter::CreateSP(this, &SGeometryCacheTimelineSplitterOverlay::GetSplitterHandlePadding, Index)))
+		[
+			SNew(SBox)
+			.Visibility(EVisibility::Visible)
+		];
+	}
+}
+
+void SGeometryCacheTimelineSplitterOverlay::OnArrangeChildren( const FGeometry& AllottedGeometry, FArrangedChildren& ArrangedChildren ) const
+{
+	FArrangedChildren SplitterChildren(ArrangedChildren.GetFilter());
+	Splitter->ArrangeChildren(AllottedGeometry, SplitterChildren);
+
+	SlotPadding.Reset();
+
+	for (int32 Index = 0; Index < SplitterChildren.Num() - 1; ++Index)
+	{
+		const FGeometry& ThisGeometry = SplitterChildren[Index].Geometry;
+		const FGeometry& NextGeometry = SplitterChildren[Index + 1].Geometry;
+
+		if (Splitter->GetOrientation() == EOrientation::Orient_Horizontal)
+		{
+			SlotPadding.Add(FMargin(
+				ThisGeometry.Position.X + static_cast<float>(ThisGeometry.GetLocalSize().X),
+				0,
+				static_cast<float>(AllottedGeometry.Size.X) - NextGeometry.Position.X,
+				0)
+			);
+		}
+		else
+		{
+			SlotPadding.Add(FMargin(
+				0,
+				ThisGeometry.Position.Y + static_cast<float>(ThisGeometry.GetLocalSize().Y),
+				0,
+				static_cast<float>(AllottedGeometry.Size.Y) - NextGeometry.Position.Y)
+			);
+		}
+	}
+
+	SOverlay::OnArrangeChildren(AllottedGeometry, ArrangedChildren);
+}
+
+FMargin SGeometryCacheTimelineSplitterOverlay::GetSplitterHandlePadding(int32 Index) const
+{
+	if (SlotPadding.IsValidIndex(Index))
+	{
+		return SlotPadding[Index];
+	}
+
+	return 0.f;
+}
+	
+FCursorReply SGeometryCacheTimelineSplitterOverlay::OnCursorQuery( const FGeometry& MyGeometry, const FPointerEvent& CursorEvent ) const
+{
+	return Splitter->OnCursorQuery(MyGeometry, CursorEvent);
+}
+
+FReply SGeometryCacheTimelineSplitterOverlay::OnMouseButtonDown( const FGeometry& MyGeometry, const FPointerEvent& MouseEvent )
+{
+	FReply Reply = Splitter->OnMouseButtonDown(MyGeometry, MouseEvent);
+	if (Reply.GetMouseCaptor().IsValid())
+	{
+		// Set us to be the mouse captor so we can forward events properly
+		Reply.CaptureMouse( SharedThis(this) );
+		SetVisibility(EVisibility::Visible);
+	}
+	return Reply;
+}
+	
+void SGeometryCacheTimelineSplitterOverlay::OnMouseCaptureLost(const FCaptureLostEvent& CaptureLostEvent)
+{
+	SetVisibility(EVisibility::SelfHitTestInvisible);
+	SOverlay::OnMouseCaptureLost(CaptureLostEvent);
+}
+
+FReply SGeometryCacheTimelineSplitterOverlay::OnMouseButtonUp( const FGeometry& MyGeometry, const FPointerEvent& MouseEvent )
+{
+	FReply Reply = Splitter->OnMouseButtonUp(MyGeometry, MouseEvent);
+	if (Reply.ShouldReleaseMouse())
+	{
+		SetVisibility(EVisibility::SelfHitTestInvisible);
+	}
+	return Reply;
+}
+
+FReply SGeometryCacheTimelineSplitterOverlay::OnMouseMove( const FGeometry& MyGeometry, const FPointerEvent& MouseEvent )
+{
+	return Splitter->OnMouseMove(MyGeometry, MouseEvent);
+}
+
+void SGeometryCacheTimelineSplitterOverlay::OnMouseLeave( const FPointerEvent& MouseEvent )
+{
+	return Splitter->OnMouseLeave(MouseEvent);
+}
+
diff --git a/Engine/Plugins/Runtime/GeometryCache/Source/GeometryCacheEd/Private/Widgets/SGeometryCacheTimelineSplitterOverlay.h b/Engine/Plugins/Runtime/GeometryCache/Source/GeometryCacheEd/Private/Widgets/SGeometryCacheTimelineSplitterOverlay.h
new file mode 100644
--- /dev/null
+++ b/Engine/Plugins/Runtime/GeometryCache/Source/GeometryCacheEd/Private/Widgets/SGeometryCacheTimelineSplitterOverlay.h
@@ -0,0 +1,41 @@
+// Copyright Epic Games, Inc. All Rights Reserved.
+
+#pragma once
+
+#include "Containers/Array.h"
+#include "HAL/Platform.h"
+#include "Input/CursorReply.h"
+#include "Input/Reply.h"
+#include "Layout/Margin.h"
+#include "Templates/SharedPointer.h"
+#include "Widgets/Layout/SSplitter.h"
+#include "Widgets/SOverlay.h"
+
+class FArrangedChildren;
+struct FCaptureLostEvent;
+struct FGeometry;
+struct FPointerEvent;
+
+/** Splitter used on the anim timeline as an overlay. Input is disabled on all areas except the draggable positions */
+class SGeometryCacheTimelineSplitterOverlay : public SOverlay
+{
+public:
+	typedef SSplitter::FArguments FArguments;
+
+	void Construct( const FArguments& InArgs );
+
+	/** SwWidget interface */
+	virtual void OnArrangeChildren( const FGeometry& AllottedGeometry, FArrangedChildren& ArrangedChildren ) const override;
+	virtual FCursorReply OnCursorQuery( const FGeometry& MyGeometry, const FPointerEvent& CursorEvent ) const override;
+	virtual FReply OnMouseButtonDown( const FGeometry& MyGeometry, const FPointerEvent& MouseEvent ) override;
+	virtual void OnMouseCaptureLost(const FCaptureLostEvent& CaptureLostEvent) override;
+	virtual FReply OnMouseButtonUp( const FGeometry& MyGeometry, const FPointerEvent& MouseEvent ) override;
+	virtual FReply OnMouseMove( const FGeometry& MyGeometry, const FPointerEvent& MouseEvent ) override;
+	virtual void OnMouseLeave( const FPointerEvent& MouseEvent ) override;
+
+private:
+	FMargin GetSplitterHandlePadding(int32 Index) const;
+
+	TSharedPtr<SSplitter> Splitter;
+	mutable TArray<FMargin> SlotPadding;
+};
diff --git a/Engine/Plugins/Runtime/GeometryCache/Source/GeometryCacheEd/Private/Widgets/SGeometryCacheTimelineTransportControls.cpp b/Engine/Plugins/Runtime/GeometryCache/Source/GeometryCacheEd/Private/Widgets/SGeometryCacheTimelineTransportControls.cpp
new file mode 100644
--- /dev/null
+++ b/Engine/Plugins/Runtime/GeometryCache/Source/GeometryCacheEd/Private/Widgets/SGeometryCacheTimelineTransportControls.cpp
@@ -0,0 +1,165 @@
+// Copyright Epic Games, Inc. All Rights Reserved.
+
+#include "SGeometryCacheTimelineTransportControls.h"
+#include "GeometryCacheComponent.h"
+#include "GeometryCacheTimelineBindingAsset.h"
+#include "Modules/ModuleManager.h"
+#include "STimelineCustomTransportControl.h"
+
+void SGeometryCacheTimelineTransportControls::Construct(const FArguments& InArgs, const TSharedRef<FGeometryCacheTimelineBindingAsset>& InBindingAsset)
+{
+	WeakBindingAsset = InBindingAsset;
+
+	FGeometryCacheTimelineTransportControlArgs TransportControlArgs;
+	TransportControlArgs.OnForwardPlay = FOnClicked::CreateSP(this, &SGeometryCacheTimelineTransportControls::OnClick_Forward);
+	TransportControlArgs.OnBackwardPlay = FOnClicked::CreateSP(this, &SGeometryCacheTimelineTransportControls::OnClick_Backward);
+	TransportControlArgs.OnForwardStep = FOnClicked::CreateSP(this, &SGeometryCacheTimelineTransportControls::OnClick_Forward_Step);
+	TransportControlArgs.OnBackwardStep = FOnClicked::CreateSP(this, &SGeometryCacheTimelineTransportControls::OnClick_Backward_Step);
+	TransportControlArgs.OnForwardEnd = FOnClicked::CreateSP(this, &SGeometryCacheTimelineTransportControls::OnClick_Forward_End);
+	TransportControlArgs.OnBackwardEnd = FOnClicked::CreateSP(this, &SGeometryCacheTimelineTransportControls::OnClick_Backward_End);
+	TransportControlArgs.OnToggleLooping = FOnClicked::CreateSP(this, &SGeometryCacheTimelineTransportControls::OnClick_ToggleLoop);
+	TransportControlArgs.OnGetLooping = FOnGetLooping::CreateSP(this, &SGeometryCacheTimelineTransportControls::IsLoopStatusOn);
+	TransportControlArgs.OnGetPlaybackMode = FOnGetPlaybackMode::CreateSP(this, &SGeometryCacheTimelineTransportControls::GetPlaybackMode);
+	
+	ChildSlot
+	[
+		SNew(STimelineCustomTransportControl)
+			.TransportArgs(TransportControlArgs)
+	];
+}
+
+FReply SGeometryCacheTimelineTransportControls::OnClick_Forward_Step()
+{
+	if (WeakBindingAsset.IsValid())
+	{
+		WeakBindingAsset.Pin()->GetPreviewComponent()->StepForward();
+	}
+	return FReply::Handled();
+}
+
+FReply SGeometryCacheTimelineTransportControls::OnClick_Forward_End()
+{
+	if (WeakBindingAsset.IsValid())
+	{
+		WeakBindingAsset.Pin()->GetPreviewComponent()->ForwardEnd();
+	}
+	return FReply::Handled();
+}
+
+FReply SGeometryCacheTimelineTransportControls::OnClick_Backward_Step()
+{
+	if (WeakBindingAsset.IsValid())
+	{
+		WeakBindingAsset.Pin()->GetPreviewComponent()->StepBackward();
+	}
+	return FReply::Handled();
+}
+
+FReply SGeometryCacheTimelineTransportControls::OnClick_Backward_End()
+{
+	if (WeakBindingAsset.IsValid())
+	{
+		WeakBindingAsset.Pin()->GetPreviewComponent()->BackwardEnd();
+	}
+	return FReply::Handled();
+}
+
+FReply SGeometryCacheTimelineTransportControls::OnClick_Forward()
+{
+	if (WeakBindingAsset.IsValid())
+	{
+		TWeakObjectPtr<UGeometryCacheComponent> GeometryCacheComponent = WeakBindingAsset.Pin()->GetPreviewComponent();
+		const bool bIsPlaying = GeometryCacheComponent->IsPlaying();
+		const bool bIsPlayingReversed = GeometryCacheComponent->IsPlayingReversed();
+
+		if (bIsPlayingReversed && bIsPlaying)
+		{
+			GeometryCacheComponent->Play();
+		}
+		else if (bIsPlaying)
+		{
+			GeometryCacheComponent->Stop();
+		}
+		else
+		{
+			if (GeometryCacheComponent->GetElapsedTime() >= GeometryCacheComponent->GetDuration())
+			{
+				GeometryCacheComponent->SetCurrentTime(0.0f);
+			}
+
+			GeometryCacheComponent->Play();
+		}
+	}
+	return FReply::Handled();
+}
+
+FReply SGeometryCacheTimelineTransportControls::OnClick_Backward()
+{
+	if (WeakBindingAsset.IsValid())
+	{
+		TWeakObjectPtr<UGeometryCacheComponent> GeometryCacheComponent = WeakBindingAsset.Pin()->GetPreviewComponent();
+		const bool bIsPlaying = GeometryCacheComponent->IsPlaying();
+		const bool bIsPlayingReversed = GeometryCacheComponent->IsPlayingReversed();
+
+		if (!bIsPlayingReversed && bIsPlaying)
+		{
+			GeometryCacheComponent->PlayReversed();
+		}
+		else if(bIsPlaying)
+		{
+			GeometryCacheComponent->Stop();
+		}
+		else
+		{
+			if (GeometryCacheComponent->GetElapsedTime() <= 0.0f)
+			{
+				GeometryCacheComponent->SetCurrentTime(GeometryCacheComponent->GetDuration());
+			}
+
+			GeometryCacheComponent->PlayReversed();
+		}
+	}
+
+	return FReply::Handled();
+}
+
+FReply SGeometryCacheTimelineTransportControls::OnClick_ToggleLoop()
+{
+	if (WeakBindingAsset.IsValid())
+	{
+		WeakBindingAsset.Pin()->GetPreviewComponent()->ToggleLooping();
+	}
+	
+	return FReply::Handled();
+}
+
+bool SGeometryCacheTimelineTransportControls::IsLoopStatusOn() const
+{
+	if (WeakBindingAsset.IsValid())
+	{
+		return WeakBindingAsset.Pin()->GetPreviewComponent()->IsLooping();
+	}
+	return false;
+}
+
+EPlaybackMode::Type SGeometryCacheTimelineTransportControls::GetPlaybackMode() const
+{
+	if (WeakBindingAsset.IsValid())
+	{
+		TWeakObjectPtr<UGeometryCacheComponent> GeometryCacheComponent = WeakBindingAsset.Pin()->GetPreviewComponent();
+		if (GeometryCacheComponent.IsValid())
+		{
+			if (GeometryCacheComponent->IsPlaying())
+			{
+				if (GeometryCacheComponent->IsPlayingReversed())
+				{
+					return EPlaybackMode::PlayingReverse;
+				}
+
+				return EPlaybackMode::PlayingForward;
+			}
+		}
+	}
+	
+	return EPlaybackMode::Stopped;
+}
\ No newline at end of file
diff --git a/Engine/Plugins/Runtime/GeometryCache/Source/GeometryCacheEd/Private/Widgets/SGeometryCacheTimelineTransportControls.h b/Engine/Plugins/Runtime/GeometryCache/Source/GeometryCacheEd/Private/Widgets/SGeometryCacheTimelineTransportControls.h
new file mode 100644
--- /dev/null
+++ b/Engine/Plugins/Runtime/GeometryCache/Source/GeometryCacheEd/Private/Widgets/SGeometryCacheTimelineTransportControls.h
@@ -0,0 +1,46 @@
+// Copyright Epic Games, Inc. All Rights Reserved.
+
+#pragma once
+
+#include "CoreMinimal.h"
+#include "Input/Reply.h"
+#include "ITransportControl.h"
+#include "Widgets/DeclarativeSyntaxSupport.h"
+#include "Widgets/SCompoundWidget.h"
+
+class FGeometryCacheTimelineBindingAsset;
+
+class SGeometryCacheTimelineTransportControls : public SCompoundWidget
+{
+public:
+	SLATE_BEGIN_ARGS(SGeometryCacheTimelineTransportControls) {}
+
+	SLATE_END_ARGS()
+
+	void Construct(const FArguments& InArgs, const TSharedRef<FGeometryCacheTimelineBindingAsset>& InBindingAsset);
+
+private:
+	
+	FReply OnClick_Forward_Step();
+
+	FReply OnClick_Forward_End();
+
+	FReply OnClick_Backward_Step();
+
+	FReply OnClick_Backward_End();
+
+	FReply OnClick_Forward();
+
+	FReply OnClick_Backward();
+
+	FReply OnClick_ToggleLoop();
+
+	bool IsLoopStatusOn() const;
+
+	EPlaybackMode::Type GetPlaybackMode() const;
+
+private:
+	TWeakPtr<FGeometryCacheTimelineBindingAsset> WeakBindingAsset;
+
+	
+};
\ No newline at end of file
diff --git a/Engine/Plugins/Runtime/GeometryCache/Source/GeometryCacheEd/Private/Widgets/STimelineCustomTransportControl.cpp b/Engine/Plugins/Runtime/GeometryCache/Source/GeometryCacheEd/Private/Widgets/STimelineCustomTransportControl.cpp
new file mode 100644
--- /dev/null
+++ b/Engine/Plugins/Runtime/GeometryCache/Source/GeometryCacheEd/Private/Widgets/STimelineCustomTransportControl.cpp
@@ -0,0 +1,258 @@
+// Copyright Epic Games, Inc. All Rights Reserved.
+
+#include "STimelineCustomTransportControl.h"
+#include "Containers/Array.h"
+#include "Delegates/Delegate.h"
+#include "Framework/SlateDelegates.h"
+#include "Internationalization/Internationalization.h"
+#include "Layout/Children.h"
+#include "Layout/Clipping.h"
+#include "Layout/Margin.h"
+#include "Layout/Visibility.h"
+#include "Misc/Attribute.h"
+#include "Misc/EnumRange.h"
+#include "SlateGlobals.h"
+#include "SlateOptMacros.h"
+#include "SlotBase.h"
+#include "Styling/AppStyle.h"
+#include "Styling/ISlateStyle.h"
+#include "Styling/SlateTypes.h"
+#include "Types/WidgetActiveTimerDelegate.h"
+#include "Widgets/Images/SImage.h"
+#include "Widgets/Input/SButton.h"
+#include "Widgets/SBoxPanel.h"
+#include "Widgets/SCompoundWidget.h"
+
+#define LOCTEXT_NAMESPACE "STransportControlCustom"
+
+namespace TransportControlConstants
+{
+	constexpr const int NumTransportControlButtons = 7;
+}
+
+BEGIN_SLATE_FUNCTION_BUILD_OPTIMIZATION
+TSharedPtr<SWidget> STimelineCustomTransportControl::MakeTransportControlWidget(ETransportControlWidgetType WidgetType, bool bAreButtonsFocusable, const FOnMakeTransportWidget& MakeCustomWidgetDelegate)
+{
+	switch (WidgetType)
+	{
+	case ETransportControlWidgetType::BackwardEnd:
+		return SNew(SButton)
+			.ButtonStyle(FAppStyle::Get(), "Animation.PlayControlsButton")
+			.OnClicked(TransportControlArgs.OnBackwardEnd)
+			.Visibility(TransportControlArgs.OnBackwardEnd.IsBound() ? EVisibility::Visible : EVisibility::Collapsed)
+			.ToolTipText(LOCTEXT("ToFront", "To Front"))
+			.ContentPadding(0.0f)
+			.IsFocusable(bAreButtonsFocusable)
+			[
+				SNew(SImage)
+					.ColorAndOpacity(FSlateColor::UseSubduedForeground())
+					.Image(FAppStyle::Get().GetBrush("Animation.Backward_End"))
+			];
+	case ETransportControlWidgetType::BackwardStep:
+		return SNew(SButton)
+			.ButtonStyle(FAppStyle::Get(), "Animation.PlayControlsButton")
+			.OnClicked(TransportControlArgs.OnBackwardStep)
+			.Visibility(TransportControlArgs.OnBackwardStep.IsBound() ? EVisibility::Visible : EVisibility::Collapsed)
+			.ToolTipText(LOCTEXT("ToPrevious", "To Previous"))
+			.ContentPadding(0.0f)
+			.IsFocusable(bAreButtonsFocusable)
+			[
+				SNew(SImage)
+					.ColorAndOpacity(FSlateColor::UseSubduedForeground())
+					.Image(FAppStyle::Get().GetBrush("Animation.Backward_Step"))
+			];
+	case ETransportControlWidgetType::BackwardPlay:
+		return SAssignNew(BackwardPlayButton, SButton)
+			.OnClicked(TransportControlArgs.OnBackwardPlay)
+			.Visibility(TransportControlArgs.OnBackwardPlay.IsBound() ? EVisibility::Visible : EVisibility::Collapsed)
+			.ToolTipText(LOCTEXT("Reverse", "Reverse"))
+			.ButtonStyle(FAppStyle::Get(), "Animation.PlayControlsButton")
+			.ContentPadding(0.0f)
+			.IsFocusable(bAreButtonsFocusable)
+			[
+				SNew(SImage)
+					.ColorAndOpacity(FSlateColor::UseSubduedForeground())
+					.Image(this, &STimelineCustomTransportControl::GetBackwardStatusIcon)
+			];
+	case ETransportControlWidgetType::ForwardPlay:
+		return SAssignNew(ForwardPlayButton, SButton)
+			.OnClicked(TransportControlArgs.OnForwardPlay)
+			.Visibility(TransportControlArgs.OnForwardPlay.IsBound() ? EVisibility::Visible : EVisibility::Collapsed)
+			.ToolTipText(this, &STimelineCustomTransportControl::GetForwardStatusTooltip)
+			.ButtonStyle(FAppStyle::Get(), "Animation.PlayControlsButton")
+			.ContentPadding(0.0f)
+			.IsFocusable(bAreButtonsFocusable)
+			[
+				SNew(SImage)
+					.ColorAndOpacity(FSlateColor::UseSubduedForeground())
+					.Image(this, &STimelineCustomTransportControl::GetForwardStatusIcon)
+			];
+	case ETransportControlWidgetType::ForwardStep:
+		return SNew(SButton)
+			.ButtonStyle(FAppStyle::Get(), "Animation.PlayControlsButton")
+			.OnClicked(TransportControlArgs.OnForwardStep)
+			.Visibility(TransportControlArgs.OnForwardStep.IsBound() ? EVisibility::Visible : EVisibility::Collapsed)
+			.ToolTipText(LOCTEXT("ToNext", "To Next"))
+			.ContentPadding(0.0f)
+			.IsFocusable(bAreButtonsFocusable)
+			[
+				SNew(SImage)
+					.ColorAndOpacity(FSlateColor::UseSubduedForeground())
+					.Image(FAppStyle::Get().GetBrush("Animation.Forward_Step"))
+			];
+	case ETransportControlWidgetType::ForwardEnd:
+		return SNew(SButton)
+			.ButtonStyle(FAppStyle::Get(), "Animation.PlayControlsButton")
+			.OnClicked(TransportControlArgs.OnForwardEnd)
+			.Visibility(TransportControlArgs.OnForwardEnd.IsBound() ? EVisibility::Visible : EVisibility::Collapsed)
+			.ToolTipText(LOCTEXT("ToEnd", "To End"))
+			.ContentPadding(0.0f)
+			.IsFocusable(bAreButtonsFocusable)
+			[
+				SNew(SImage)
+					.ColorAndOpacity(FSlateColor::UseSubduedForeground())
+					.Image(FAppStyle::Get().GetBrush("Animation.Forward_End"))
+			];
+	case ETransportControlWidgetType::Loop:
+		return SAssignNew(LoopButton, SButton)
+			.OnClicked(TransportControlArgs.OnToggleLooping)
+			.Visibility(TransportControlArgs.OnGetLooping.IsBound() ? EVisibility::Visible : EVisibility::Collapsed)
+			.ToolTipText(this, &STimelineCustomTransportControl::GetLoopStatusTooltip)
+			.ButtonStyle(FAppStyle::Get(), "Animation.PlayControlsButton")
+			.ContentPadding(0.0f)
+			.IsFocusable(bAreButtonsFocusable)
+			[
+				SNew(SImage)
+					.ColorAndOpacity(FSlateColor::UseSubduedForeground())
+					.Image(this, &STimelineCustomTransportControl::GetLoopStatusIcon)
+			];
+	default:
+	{
+		return SNullWidget::NullWidget;
+	}
+	}
+
+	return nullptr;
+}
+
+void STimelineCustomTransportControl::Construct(const FArguments& InArgs)
+{
+	TransportControlArgs = InArgs._TransportArgs;
+
+	TSharedRef<SHorizontalBox> HorizontalBox = SNew(SHorizontalBox)
+		.Clipping(EWidgetClipping::ClipToBounds);
+
+	TArray<ETransportControlWidgetType, TFixedAllocator<TransportControlConstants::NumTransportControlButtons>> ButtonWidgetTypes = {
+		ETransportControlWidgetType::BackwardEnd,
+		ETransportControlWidgetType::BackwardStep,
+		ETransportControlWidgetType::BackwardPlay,
+		ETransportControlWidgetType::ForwardPlay,
+		ETransportControlWidgetType::ForwardStep,
+		ETransportControlWidgetType::ForwardEnd,
+		ETransportControlWidgetType::Loop,
+	};
+
+	for (ETransportControlWidgetType WidgetType : ButtonWidgetTypes)
+	{
+		TSharedPtr<SWidget> Widget = MakeTransportControlWidget(WidgetType, InArgs._TransportArgs.bAreButtonsFocusable);
+		if (Widget.IsValid())
+		{
+			HorizontalBox->AddSlot()
+				.AutoWidth()
+				.HAlign(HAlign_Center)
+				.VAlign(VAlign_Center)
+				[
+					Widget.ToSharedRef()
+				];
+		}
+	}
+
+	ChildSlot
+		[
+			HorizontalBox
+		];
+}
+
+bool STimelineCustomTransportControl::IsTickable() const
+{
+	return false;
+}
+
+void STimelineCustomTransportControl::Tick(float DeltaTime)
+{
+}
+
+const FSlateBrush* STimelineCustomTransportControl::GetForwardStatusIcon() const
+{
+	EPlaybackMode::Type PlaybackMode = EPlaybackMode::Stopped;
+	if (TransportControlArgs.OnGetPlaybackMode.IsBound())
+	{
+		PlaybackMode = TransportControlArgs.OnGetPlaybackMode.Execute();
+	}
+
+	if (PlaybackMode == EPlaybackMode::PlayingForward)
+	{
+		return FAppStyle::Get().GetBrush("Animation.Pause");
+	}
+
+	return FAppStyle::Get().GetBrush("Animation.Forward");
+}
+
+FText STimelineCustomTransportControl::GetForwardStatusTooltip() const
+{
+	if (TransportControlArgs.OnGetPlaybackMode.IsBound() &&
+		TransportControlArgs.OnGetPlaybackMode.Execute() == EPlaybackMode::PlayingForward)
+	{
+		return LOCTEXT("Pause", "Pause");
+	}
+
+	return LOCTEXT("Play", "Play");
+}
+
+const FSlateBrush* STimelineCustomTransportControl::GetBackwardStatusIcon() const
+{
+	if (TransportControlArgs.OnGetPlaybackMode.IsBound() &&
+		TransportControlArgs.OnGetPlaybackMode.Execute() == EPlaybackMode::PlayingReverse)
+	{
+		return FAppStyle::Get().GetBrush("Animation.Pause");
+	}
+
+	return FAppStyle::Get().GetBrush("Animation.Backward");
+}
+
+const FSlateBrush* STimelineCustomTransportControl::GetLoopStatusIcon() const
+{
+	if (TransportControlArgs.OnGetLooping.IsBound() &&
+		TransportControlArgs.OnGetLooping.Execute())
+	{
+		return FAppStyle::Get().GetBrush("Animation.Loop.Enabled");
+	}
+
+	return FAppStyle::Get().GetBrush("Animation.Loop.Disabled");
+}
+
+FText STimelineCustomTransportControl::GetLoopStatusTooltip() const
+{
+	if (TransportControlArgs.OnGetLooping.IsBound() &&
+		TransportControlArgs.OnGetLooping.Execute())
+	{
+		return LOCTEXT("Looping", "Looping");
+	}
+
+	return LOCTEXT("NoLooping", "No Looping");
+}
+
+EActiveTimerReturnType STimelineCustomTransportControl::TickPlayback(double InCurrentTime, float InDeltaTime)
+{
+	TransportControlArgs.OnTickPlayback.Execute(InCurrentTime, InDeltaTime);
+	return EActiveTimerReturnType::Continue;
+}
+
+FReply STimelineCustomTransportControl::OnToggleLooping()
+{
+	return FReply::Handled();
+}
+
+END_SLATE_FUNCTION_BUILD_OPTIMIZATION
+
+#undef LOCTEXT_NAMESPACE
\ No newline at end of file
diff --git a/Engine/Plugins/Runtime/GeometryCache/Source/GeometryCacheEd/Private/Widgets/STimelineCustomTransportControl.h b/Engine/Plugins/Runtime/GeometryCache/Source/GeometryCacheEd/Private/Widgets/STimelineCustomTransportControl.h
new file mode 100644
--- /dev/null
+++ b/Engine/Plugins/Runtime/GeometryCache/Source/GeometryCacheEd/Private/Widgets/STimelineCustomTransportControl.h
@@ -0,0 +1,105 @@
+// Copyright Epic Games, Inc. All Rights Reserved.
+
+#pragma once
+
+#include "Input/Reply.h"
+#include "Internationalization/Text.h"
+#include "ITransportControl.h"
+#include "Stats/Stats2.h"
+#include "Templates/SharedPointer.h"
+#include "TickableEditorObject.h"
+#include "Types/SlateEnums.h"
+#include "Widgets/DeclarativeSyntaxSupport.h"
+#include "Widgets/SWidget.h"
+
+class FActiveTimerHandle;
+class SButton;
+struct FSlateBrush;
+
+DECLARE_DELEGATE_RetVal(bool, FOnGetLooping)
+DECLARE_DELEGATE_RetVal(EPlaybackMode::Type, FOnGetPlaybackMode)
+DECLARE_DELEGATE_TwoParams(FOnTickPlayback, double /*InCurrentTime*/, float /*InDeltaTime*/)
+
+struct FGeometryCacheTimelineTransportControlArgs
+{
+	FGeometryCacheTimelineTransportControlArgs()
+		: OnForwardPlay()
+		, OnBackwardPlay()
+		, OnForwardStep()
+		, OnBackwardStep()
+		, OnForwardEnd()
+		, OnBackwardEnd()
+		, OnToggleLooping()
+		, OnGetLooping()
+		, OnGetPlaybackMode()
+		, OnTickPlayback()
+		, bAreButtonsFocusable(true)
+	{}
+
+	FOnClicked OnForwardPlay;
+	FOnClicked OnBackwardPlay;
+	FOnClicked OnForwardStep;
+	FOnClicked OnBackwardStep;
+	FOnClicked OnForwardEnd;
+	FOnClicked OnBackwardEnd;
+	FOnClicked OnToggleLooping;
+	FOnGetLooping OnGetLooping;
+	FOnGetPlaybackMode OnGetPlaybackMode;
+	FOnTickPlayback OnTickPlayback;
+	bool bAreButtonsFocusable;
+};
+
+class STimelineCustomTransportControl : public ITransportControl, public FTickableEditorObject
+{
+public:
+	SLATE_BEGIN_ARGS(STimelineCustomTransportControl) 
+		:_TransportArgs()
+		{}
+		SLATE_ARGUMENT(FGeometryCacheTimelineTransportControlArgs, TransportArgs)
+	SLATE_END_ARGS()
+
+	/**
+	 * Construct the widget
+	 *
+	 * @param InArgs   A declaration from which to construct the widget
+	 */
+	void Construct(const FArguments& InArgs);
+
+	virtual ~STimelineCustomTransportControl() {}
+
+	using SWidget::Tick;
+
+	// Begin FTickableObjectBase implementation
+	virtual bool IsTickable() const override;
+	virtual void Tick(float DeltaTime) override;
+	virtual TStatId GetStatId() const override { RETURN_QUICK_DECLARE_CYCLE_STAT(STimelineCustomTransportControl, STATGROUP_Tickables); }
+	// End FTickableObjectBase
+
+private:
+	const FSlateBrush* GetForwardStatusIcon() const;
+	FText GetForwardStatusTooltip() const;
+	const FSlateBrush* GetBackwardStatusIcon() const;
+	const FSlateBrush* GetLoopStatusIcon() const;
+	FText GetLoopStatusTooltip() const;
+
+	/** Executes the OnTickPlayback delegate */
+	EActiveTimerReturnType TickPlayback(double InCurrentTime, float InDeltaTime);
+
+	FReply OnToggleLooping();
+
+	/** Make default transport control widgets */
+	TSharedPtr<SWidget> MakeTransportControlWidget(ETransportControlWidgetType WidgetType, bool bAreButtonsFocusable, const FOnMakeTransportWidget& MakeCustomWidgetDelegate = FOnMakeTransportWidget());
+
+private:
+	/** The handle to the active timer */
+	TWeakPtr<FActiveTimerHandle> ActiveTimerHandle;
+
+	/** Whether the active timer is currently registered */
+	bool bIsActiveTimerRegistered;
+
+	FGeometryCacheTimelineTransportControlArgs TransportControlArgs;
+
+	TSharedPtr<SButton> ForwardPlayButton;
+	TSharedPtr<SButton> BackwardPlayButton;
+	TSharedPtr<SButton> LoopButton;
+};
\ No newline at end of file
