Fix FGeometryCacheSceneProxy not correctly handling ray tracing geometry being evicted.
- also change behavior so ray tracing geometry is released when it doesn't have any Segments/TotalPrimitiveCount is zero.

diff --git a/Engine/Plugins/Runtime/GeometryCache/Source/GeometryCache/Private/GeometryCacheSceneProxy.cpp b/Engine/Plugins/Runtime/GeometryCache/Source/GeometryCache/Private/GeometryCacheSceneProxy.cpp
--- a/Engine/Plugins/Runtime/GeometryCache/Source/GeometryCache/Private/GeometryCacheSceneProxy.cpp
+++ b/Engine/Plugins/Runtime/GeometryCache/Source/GeometryCache/Private/GeometryCacheSceneProxy.cpp
@@ -680,21 +680,18 @@ void FGeometryCacheSceneProxy::UpdateAnimation(FRHICommandListBase& RHICmdList,
 				Section->RayTracingGeometry.Initializer.IndexBuffer = Section->IndexBuffer.IndexBufferRHI;
 				Section->RayTracingGeometry.Initializer.TotalPrimitiveCount = TotalPrimitiveCount;
 
-				if (Segments.Num() > 0)
+				if (bRequireRecreate)
 				{
-					if (bRequireRecreate)
-					{
-						Section->RayTracingGeometry.UpdateRHI(RHICmdList);
-					}
-					else
-					{
-						// Request full build on same geometry because data might have changed to much for update call?
-						FRayTracingGeometryBuildParams BuildParams;
-						BuildParams.Geometry = Section->RayTracingGeometry.GetRHI();
-						BuildParams.BuildMode = EAccelerationStructureBuildMode::Build;
-						BuildParams.Segments = Section->RayTracingGeometry.Initializer.Segments;
-						FRHIComputeCommandList::Get(RHICmdList).BuildAccelerationStructures(MakeArrayView(&BuildParams, 1));
-					}
+					Section->RayTracingGeometry.UpdateRHI(RHICmdList);
+				}
+				else if(Section->RayTracingGeometry.IsValid() && !Section->RayTracingGeometry.IsEvicted())
+				{
+					// Request full build on same geometry because data might have changed to much for update call?
+					FRayTracingGeometryBuildParams BuildParams;
+					BuildParams.Geometry = Section->RayTracingGeometry.GetRHI();
+					BuildParams.BuildMode = EAccelerationStructureBuildMode::Build;
+					BuildParams.Segments = Section->RayTracingGeometry.Initializer.Segments;
+					FRHIComputeCommandList::Get(RHICmdList).BuildAccelerationStructures(MakeArrayView(&BuildParams, 1));
 				}
 			}
 		}
