[BUGFIX] GeometryCacheAssetEditorToolkit - Localization Build Fails

diff --git a/Engine/Plugins/Runtime/GeometryCache/Source/GeometryCacheEd/Private/GeometryCacheAssetEditorToolkit.cpp b/Engine/Plugins/Runtime/GeometryCache/Source/GeometryCacheEd/Private/GeometryCacheAssetEditorToolkit.cpp
--- a/Engine/Plugins/Runtime/GeometryCache/Source/GeometryCacheEd/Private/GeometryCacheAssetEditorToolkit.cpp
+++ b/Engine/Plugins/Runtime/GeometryCache/Source/GeometryCacheEd/Private/GeometryCacheAssetEditorToolkit.cpp
@@ -41,7 +41,7 @@ void FGeometryCacheAssetEditorToolkit::RegisterTabSpawners(const TSharedRef<clas
 		.SetIcon(FSlateIcon(FAppStyle::GetAppStyleSetName(), "LevelEditor.Tabs.Render"));
 
 	InTabManager->RegisterTabSpawner(TabId_AssetProperties, FOnSpawnTab::CreateSP(this, &FGeometryCacheAssetEditorToolkit::SpawnTab_AssetProperties))
-		.SetDisplayName(LOCTEXT("AssetPropertiesTab", "Details"))
+		.SetDisplayName(LOCTEXT("AssetPropertiesTab", "Properties"))
 		.SetGroup(WorkspaceMenuCategory.ToSharedRef())
 		.SetIcon(FSlateIcon(FAppStyle::GetAppStyleSetName(), "LevelEditor.Tabs.Details"));
 
