Fixed more ISPC per-kernel variables not having or not using the preprocessor define for the default value.

diff --git a/Engine/Plugins/FX/Niagara/Source/Niagara/Classes/NiagaraDataInterfaceVectorField.h b/Engine/Plugins/FX/Niagara/Source/Niagara/Classes/NiagaraDataInterfaceVectorField.h
--- a/Engine/Plugins/FX/Niagara/Source/Niagara/Classes/NiagaraDataInterfaceVectorField.h
+++ b/Engine/Plugins/FX/Niagara/Source/Niagara/Classes/NiagaraDataInterfaceVectorField.h
@@ -109,3 +109,14 @@ struct FNiagaraDataInterfaceProxyVectorField : public FNiagaraDataInterfaceProxy
 		return 0;
 	}
 };
+
+#if !defined(NIAGARA_VECTOR_FIELD_ISPC_ENABLED_DEFAULT)
+#define NIAGARA_VECTOR_FIELD_ISPC_ENABLED_DEFAULT 1
+#endif
+
+// Support run-time toggling on supported platforms in non-shipping configurations
+#if !INTEL_ISPC || UE_BUILD_SHIPPING
+static constexpr bool GNiagaraVectorFieldUseIspc = INTEL_ISPC && NIAGARA_VECTOR_FIELD_ISPC_ENABLED_DEFAULT;
+#else
+extern bool GNiagaraVectorFieldUseIspc;
+#endif
diff --git a/Engine/Plugins/FX/Niagara/Source/Niagara/Private/NiagaraDataInterfaceVectorField.cpp b/Engine/Plugins/FX/Niagara/Source/Niagara/Private/NiagaraDataInterfaceVectorField.cpp
--- a/Engine/Plugins/FX/Niagara/Source/Niagara/Private/NiagaraDataInterfaceVectorField.cpp
+++ b/Engine/Plugins/FX/Niagara/Source/Niagara/Private/NiagaraDataInterfaceVectorField.cpp
@@ -30,10 +30,8 @@ static const FName GetVectorFieldBoundsName("FieldBounds");
 #if INTEL_ISPC
 static_assert(sizeof(ispc::FVector) == sizeof(FVector), "sizeof(ispc::FVector) != sizeof(FVector)");
 
-#if UE_BUILD_SHIPPING
-const bool GNiagaraVectorFieldUseIspc = true;
-#else
-bool GNiagaraVectorFieldUseIspc = true;
+#if !UE_BUILD_SHIPPING
+bool GNiagaraVectorFieldUseIspc = NIAGARA_VECTOR_FIELD_ISPC_ENABLED_DEFAULT;
 static FAutoConsoleVariableRef CVarNiagaraVectorFieldUseIspc(
 	TEXT("fx.NiagaraVectorFieldUseIspc"),
 	GNiagaraVectorFieldUseIspc,
diff --git a/Engine/Plugins/Runtime/GeometryCache/Source/GeometryCache/Private/GeometryCacheSceneProxy.cpp b/Engine/Plugins/Runtime/GeometryCache/Source/GeometryCache/Private/GeometryCacheSceneProxy.cpp
--- a/Engine/Plugins/Runtime/GeometryCache/Source/GeometryCache/Private/GeometryCacheSceneProxy.cpp
+++ b/Engine/Plugins/Runtime/GeometryCache/Source/GeometryCache/Private/GeometryCacheSceneProxy.cpp
@@ -30,10 +30,8 @@ static_assert(sizeof(uint32) == sizeof(FPackedNormal), "sizeof(uint32) != sizeof
 static_assert(sizeof(ispc::FVector2f) == sizeof(FVector2f), "sizeof(ispc::FVector2f) != sizeof(FVector2f)");
 static_assert(sizeof(ispc::FVector3f) == sizeof(FVector3f), "sizeof(ispc::FVector3f) != sizeof(FVector3f)");
 
-#if UE_BUILD_SHIPPING
-const bool GGeometryCacheSceneProxyUseIspc = true;
-#else
-bool GGeometryCacheSceneProxyUseIspc = true;
+#if !UE_BUILD_SHIPPING
+bool GGeometryCacheSceneProxyUseIspc = GEOMETRY_CACHE_SCENE_PROXY_ISPC_ENABLED_DEFAULT;
 static FAutoConsoleVariableRef CVarGeometryCacheSceneProxyUseIspc(
 	TEXT("r.GeometryCacheSceneProxy.ISPC"),
 	GGeometryCacheSceneProxyUseIspc,
diff --git a/Engine/Plugins/Runtime/GeometryCache/Source/GeometryCache/Public/GeometryCacheSceneProxy.h b/Engine/Plugins/Runtime/GeometryCache/Source/GeometryCache/Public/GeometryCacheSceneProxy.h
--- a/Engine/Plugins/Runtime/GeometryCache/Source/GeometryCache/Public/GeometryCacheSceneProxy.h
+++ b/Engine/Plugins/Runtime/GeometryCache/Source/GeometryCache/Public/GeometryCacheSceneProxy.h
@@ -414,3 +414,13 @@ private:
 #include "PrimitiveViewRelevance.h"
 #include "StaticMeshResources.h"
 #endif
+
+#if !defined(GEOMETRY_CACHE_SCENE_PROXY_ISPC_ENABLED_DEFAULT)
+#define GEOMETRY_CACHE_SCENE_PROXY_ISPC_ENABLED_DEFAULT 1
+#endif
+
+#if !INTEL_ISPC || UE_BUILD_SHIPPING
+static constexpr bool GGeometryCacheSceneProxyUseIspc = INTEL_ISPC && GEOMETRY_CACHE_SCENE_PROXY_ISPC_ENABLED_DEFAULT;
+#else
+extern bool GGeometryCacheSceneProxyUseIspc;
+#endif
