[HWRT] Fix FGeometryCacheSceneProxy::GetDynamicRayTracingInstances() using FRHICommandListImmediate::Get() instead of Collector.GetRHICommandList()

diff --git a/Engine/Plugins/Runtime/GeometryCache/Source/GeometryCache/Private/GeometryCacheSceneProxy.cpp b/Engine/Plugins/Runtime/GeometryCache/Source/GeometryCache/Private/GeometryCacheSceneProxy.cpp
--- a/Engine/Plugins/Runtime/GeometryCache/Source/GeometryCache/Private/GeometryCacheSceneProxy.cpp
+++ b/Engine/Plugins/Runtime/GeometryCache/Source/GeometryCache/Private/GeometryCacheSceneProxy.cpp
@@ -529,8 +529,6 @@ void FGeometryCacheSceneProxy::GetDynamicRayTracingInstances(FRayTracingInstance
 		return;
 	}
 
-	FRHICommandListBase& RHICmdList = FRHICommandListImmediate::Get();
-
 	for (FGeomCacheTrackProxy* TrackProxy : Tracks)
 	{
 		const FVisibilitySample& VisibilitySample = TrackProxy->GetVisibilitySample(Time, bLooping);
@@ -551,7 +549,7 @@ void FGeometryCacheSceneProxy::GetDynamicRayTracingInstances(FRayTracingInstance
 
 			FGeometryCacheVertexFactoryUserDataWrapper& UserDataWrapper = Collector.AllocateOneFrameResource<FGeometryCacheVertexFactoryUserDataWrapper>();
 			FDynamicPrimitiveUniformBuffer& DynamicPrimitiveUniformBuffer = Collector.AllocateOneFrameResource<FDynamicPrimitiveUniformBuffer>();
-			CreateMeshBatch(RHICmdList, TrackProxy, BatchInfo, UserDataWrapper, DynamicPrimitiveUniformBuffer, MeshBatch);
+			CreateMeshBatch(Collector.GetRHICommandList(), TrackProxy, BatchInfo, UserDataWrapper, DynamicPrimitiveUniformBuffer, MeshBatch);
 
 			const int32 MaterialIndex = TrackProxy->Materials.IsValidIndex(BatchInfo.MaterialIndex) ? BatchInfo.MaterialIndex : SegmentIndex;
 			MeshBatch.MaterialRenderProxy = TrackProxy->Materials[MaterialIndex]->GetRenderProxy();
