Alembic Cache Editor: Needs new Viewport toolbars
- Refactored old toolbar to use the PopulateViewportOverlays for stats text
- Added the new viewport toolbar using the BuildMenuToolbar.

#rnx

diff --git a/Engine/Plugins/Runtime/GeometryCache/Source/GeometryCacheEd/Private/Widgets/SGeometryCacheEditorViewport.cpp b/Engine/Plugins/Runtime/GeometryCache/Source/GeometryCacheEd/Private/Widgets/SGeometryCacheEditorViewport.cpp
--- a/Engine/Plugins/Runtime/GeometryCache/Source/GeometryCacheEd/Private/Widgets/SGeometryCacheEditorViewport.cpp
+++ b/Engine/Plugins/Runtime/GeometryCache/Source/GeometryCacheEd/Private/Widgets/SGeometryCacheEditorViewport.cpp
@@ -16,7 +16,9 @@
 #include "SCommonEditorViewportToolbarBase.h"
 #include "Slate/SceneViewport.h"
 #include "ThumbnailRendering/ThumbnailManager.h"
+#include "ToolMenus.h"
 #include "UnrealEdGlobals.h"
+#include "ViewportToolbar/UnrealEdViewportToolbar.h"
 #include "Widgets/Docking/SDockTab.h"
 #include "Widgets/SNullWidget.h"
 #include "Widgets/Text/SRichTextBlock.h"
@@ -237,11 +239,6 @@ TSharedRef<FEditorViewportClient> SGeometryCacheEditorViewport::MakeEditorViewpo
 
 TSharedPtr<SWidget> SGeometryCacheEditorViewport::MakeViewportToolbar()
 {
-	TAttribute<FText> StatsTextAttribute;
-	{
-		StatsTextAttribute.Bind(TAttribute<FText>::FGetter::CreateSP(this, &SGeometryCacheEditorViewport::BuildStatsText));
-	}
-
 	return SNew(SVerticalBox)
 		+ SVerticalBox::Slot()
 		.AutoHeight()
@@ -252,16 +249,90 @@ TSharedPtr<SWidget> SGeometryCacheEditorViewport::MakeViewportToolbar()
 				[
 					SNew(SCommonEditorViewportToolbarBase, SharedThis(this))
 				]
-		]
-		+ SVerticalBox::Slot()
-		.Padding(FMargin(4.0f, 3.0f, 0.0f, 0.0f))
+		];
+}
+
+
+TSharedPtr<SWidget> SGeometryCacheEditorViewport::BuildViewportToolbar()
+{
+	const FName ViewportToolbarName = "GeometryCacheEditor.ViewportToolbar";
+
+	// Register the viewport toolbar if another viewport hasn't already (it's shared).
+	if (!UToolMenus::Get()->IsMenuRegistered(ViewportToolbarName))
+	{
+		UToolMenu* const ViewportToolbarMenu = UToolMenus::Get()->RegisterMenu(
+			ViewportToolbarName, NAME_None /* parent */, EMultiBoxType::SlimHorizontalToolBar
+		);
+
+		ViewportToolbarMenu->StyleName = "ViewportToolbar";
+
+		// Add the left-aligned part of the viewport toolbar.
+		{
+			FToolMenuSection& LeftSection = ViewportToolbarMenu->FindOrAddSection("Left");
+		}
+
+		// Add the right-aligned part of the viewport toolbar.
+		{
+			FToolMenuSection& RightSection = ViewportToolbarMenu->FindOrAddSection("Right");
+			RightSection.Alignment = EToolMenuSectionAlign::Last;
+
+			// Add the "Camera" submenu.
+			{
+				// Build the menu name our Camera menu will be using so we can extend it 
+				const FName SubmenuName = UToolMenus::JoinMenuPaths(ViewportToolbarName, "Camera");
+				// Extending using UnrealEd generic entries
+				UE::UnrealEd::ExtendCameraSubmenu(SubmenuName);
+				RightSection.AddEntry(UE::UnrealEd::CreateCameraSubmenu());
+			}
+
+			// Add the "View Modes" sub menu.
+			RightSection.AddEntry(UE::UnrealEd::CreateViewModesSubmenu());
+			
+			// Add the "Show" submenu.
+			RightSection.AddEntry(UE::UnrealEd::CreateDefaultShowSubmenu());
+			
+			// Performance and Scalability Submenu
+			{
+				FToolMenuEntry Entry = UE::UnrealEd::CreatePerformanceAndScalabilitySubmenu();
+				{
+					Entry.Icon = FSlateIcon(FAppStyle::GetAppStyleSetName(), "Icons.Scalability");
+					Entry.InsertPosition.Position = EToolMenuInsertType::Last;
+					Entry.ToolBarData.LabelOverride = FText();
+					Entry.ToolBarData.ResizeParams.ClippingPriority = 800;
+				}
+
+				// Add Performance and Scalability Menu
+				RightSection.AddEntry(Entry);
+
+				const FName SubmenuName = UToolMenus::JoinMenuPaths(ViewportToolbarName, "PerformanceAndScalability");
+				UToolMenu* Submenu = UToolMenus::Get()->ExtendMenu(SubmenuName);
+				FToolMenuSection& Section = Submenu->FindOrAddSection("PerformanceAndScalability");
+
+				// Add Scalability Menu
+				Section.AddEntry(UE::UnrealEd::CreateScalabilitySubmenu());
+			}
+		}
+	}
+
+	FToolMenuContext ViewportToolbarContext;
+	{
+		ViewportToolbarContext.AppendCommandList(GetCommandList());
+
+		// Add the UnrealEd viewport toolbar context.
+		{
+			UUnrealEdViewportToolbarContext* const ContextObject = NewObject<UUnrealEdViewportToolbarContext>();
+			ContextObject->Viewport = SharedThis(this);
+			ViewportToolbarContext.AddObject(ContextObject);
+		}
+	}
+
+	
+	const TSharedRef<SWidget> NewViewportToolbar = SNew(SBox)
 		[
-			SNew(SRichTextBlock)
-				.Visibility(EVisibility::SelfHitTestInvisible)
-				.DecoratorStyleSet(&FAppStyle::Get())
-				.Text(StatsTextAttribute)
-				.TextStyle(&FAppStyle::Get().GetWidgetStyle<FTextBlockStyle>("AnimViewport.MessageText"))
+			UToolMenus::Get()->GenerateWidget(ViewportToolbarName, ViewportToolbarContext)
 		];
+	
+	return NewViewportToolbar;
 }
 
 EVisibility SGeometryCacheEditorViewport::OnGetViewportContentVisibility() const
@@ -289,6 +360,37 @@ void SGeometryCacheEditorViewport::OnFocusViewportToSelection()
 
 void SGeometryCacheEditorViewport::PopulateViewportOverlays(TSharedRef<class SOverlay> Overlay)
 {
+	SEditorViewport::PopulateViewportOverlays(Overlay);
+
+	TSharedRef<SVerticalBox> VerticalBox = SNew(SVerticalBox);
+
+	const TAttribute<FMargin> Padding = TAttribute<FMargin>::CreateLambda(
+		[]() -> FMargin
+		{
+			// The old viewport toolbar draws inside the viewport, so we need to pad to not overlap it.
+			const float TopPadding = UE::UnrealEd::ShowOldViewportToolbars() ? 29.0f : 0;
+			return FMargin(0, TopPadding, 0, 0);
+		}
+	);
+
+	Overlay->AddSlot()
+		.VAlign(VAlign_Top)
+		.HAlign(HAlign_Fill)
+		.Padding(Padding)
+		[
+			VerticalBox
+		];
+
+	VerticalBox->AddSlot()
+		.AutoHeight()
+		.Padding(FMargin(4.0f, 3.0f, 0.0f, 0.0f))
+		[
+			SNew(SRichTextBlock)
+				.Visibility(EVisibility::SelfHitTestInvisible)
+				.DecoratorStyleSet(&FAppStyle::Get())
+				.Text(this, &SGeometryCacheEditorViewport::BuildStatsText)
+				.TextStyle(&FAppStyle::Get().GetWidgetStyle<FTextBlockStyle>("AnimViewport.MessageText"))
+		];
 }
 
 bool SGeometryCacheEditorViewport::IsVisible() const
diff --git a/Engine/Plugins/Runtime/GeometryCache/Source/GeometryCacheEd/Private/Widgets/SGeometryCacheEditorViewport.h b/Engine/Plugins/Runtime/GeometryCache/Source/GeometryCacheEd/Private/Widgets/SGeometryCacheEditorViewport.h
--- a/Engine/Plugins/Runtime/GeometryCache/Source/GeometryCacheEd/Private/Widgets/SGeometryCacheEditorViewport.h
+++ b/Engine/Plugins/Runtime/GeometryCache/Source/GeometryCacheEd/Private/Widgets/SGeometryCacheEditorViewport.h
@@ -47,6 +47,7 @@ protected:
 	/** SEditorViewport interface */
 	virtual TSharedRef<FEditorViewportClient> MakeEditorViewportClient() override;
 	virtual TSharedPtr<SWidget> MakeViewportToolbar() override;
+	virtual TSharedPtr<SWidget> BuildViewportToolbar() override;
 	virtual EVisibility OnGetViewportContentVisibility() const override;
 	virtual void BindCommands() override;
 	virtual void OnFocusViewportToSelection() override;
