[Geometry Cache] turning back on raytracing for geo cache

diff --git a/Engine/Plugins/Runtime/GeometryCache/Source/GeometryCache/Private/GeometryCacheSceneProxy.cpp b/Engine/Plugins/Runtime/GeometryCache/Source/GeometryCache/Private/GeometryCacheSceneProxy.cpp
--- a/Engine/Plugins/Runtime/GeometryCache/Source/GeometryCache/Private/GeometryCacheSceneProxy.cpp
+++ b/Engine/Plugins/Runtime/GeometryCache/Source/GeometryCache/Private/GeometryCacheSceneProxy.cpp
@@ -71,7 +71,7 @@ static TAutoConsoleVariable<int32> CVarInterpolateFrames(
 
 static TAutoConsoleVariable<int32> CVarRayTracingGeometryCache(
 	TEXT("r.RayTracing.Geometry.GeometryCache"),
-	0,
+	1,
 	TEXT("Include geometry cache primitives in ray tracing effects (default = 1 (geometry cache enabled in ray tracing))"));
 
 /**
