[GeometryCache] Fixed raytracing gpu crash when the first frame of geo cahe is invisible

diff --git a/Engine/Plugins/Experimental/GeometryCacheLevelSequenceBaker/Source/GeometryCacheLevelSequenceBaker/Private/GeometryCacheLevelSequenceBaker.cpp b/Engine/Plugins/Experimental/GeometryCacheLevelSequenceBaker/Source/GeometryCacheLevelSequenceBaker/Private/GeometryCacheLevelSequenceBaker.cpp
--- a/Engine/Plugins/Experimental/GeometryCacheLevelSequenceBaker/Source/GeometryCacheLevelSequenceBaker/Private/GeometryCacheLevelSequenceBaker.cpp
+++ b/Engine/Plugins/Experimental/GeometryCacheLevelSequenceBaker/Source/GeometryCacheLevelSequenceBaker/Private/GeometryCacheLevelSequenceBaker.cpp
@@ -739,19 +739,25 @@ void FGeometryCacheLevelSequenceBaker::WriteToAsset()
 		{
 			continue;
 		}
-		
-		// Copy previous frame if the component isn't visible at this frame so that
-        // frame interpolation works correctly
+
+		// Copy neighbor frames when visibility is changing such that frame interpolation works well
 		for (const FVisibilitySample& VisibilitySample : Task.VisibilitySamples)
 		{	
 			if (!VisibilitySample.bVisible)
 			{
-				if (VisibilitySample.FrameIndex > 0)
+				const int32 Previous = VisibilitySample.FrameIndex-1 ;
+				const int32 Next = VisibilitySample.FrameIndex+1 ;
+				// Copy from next if changing to visible
+				if (Task.VisibilitySamples.IsValidIndex(Next) &&
+					Task.VisibilitySamples[Next].bVisible)
 				{
-					if (Task.VisibilitySamples[VisibilitySample.FrameIndex-1].bVisible)
-					{
-						Task.GeometrySamples[VisibilitySample.FrameIndex] = Task.GeometrySamples[VisibilitySample.FrameIndex - 1];	
-					}
+					Task.GeometrySamples[VisibilitySample.FrameIndex] = Task.GeometrySamples[Next];	
+				}
+				// Copy from previous if changing to invisible
+				else if (Task.VisibilitySamples.IsValidIndex(Previous) && 
+					Task.VisibilitySamples[Previous].bVisible)
+				{
+					Task.GeometrySamples[VisibilitySample.FrameIndex] = Task.GeometrySamples[Previous];	
 				}
 			}
 		}
diff --git a/Engine/Plugins/Runtime/GeometryCache/Source/GeometryCache/Private/GeometryCacheSceneProxy.cpp b/Engine/Plugins/Runtime/GeometryCache/Source/GeometryCache/Private/GeometryCacheSceneProxy.cpp
--- a/Engine/Plugins/Runtime/GeometryCache/Source/GeometryCache/Private/GeometryCacheSceneProxy.cpp
+++ b/Engine/Plugins/Runtime/GeometryCache/Source/GeometryCache/Private/GeometryCacheSceneProxy.cpp
@@ -191,52 +191,11 @@ FGeometryCacheSceneProxy::FGeometryCacheSceneProxy(UGeometryCacheComponent* Comp
 	{
 #if RHI_RAYTRACING
 		{
+			RayTracingDebugName = Component->GetFName();
 			ENQUEUE_RENDER_COMMAND(FGeometryCacheInitRayTracingGeometry)(
-				[SceneProxy, DebugName = Component->GetFName()](FRHICommandListImmediate& RHICmdList)
+			[SceneProxy](FRHICommandListImmediate& RHICmdList)
 			{
-				for (FGeomCacheTrackProxy* Section : SceneProxy->Tracks)
-				{
-					if (Section != nullptr)
-					{
-						const FVisibilitySample& VisibilitySample = Section->GetVisibilitySample(SceneProxy->Time, SceneProxy->bLooping);
-						if (!VisibilitySample.bVisibilityState)
-						{
-							continue;
-						}
-			
-						FRayTracingGeometryInitializer Initializer;
-						Initializer.DebugName = DebugName;
-						const int PositionBufferIndex = Section->CurrentPositionBufferIndex != -1 ? Section->CurrentPositionBufferIndex % 2 : 0;
-						Initializer.TotalPrimitiveCount = 0;
-						Initializer.GeometryType = RTGT_Triangles;
-						Initializer.bFastBuild = false;
-						Initializer.bAllowCompaction = false;
-
-						TArray<FRayTracingGeometrySegment> Segments;
-						const FGeometryCacheMeshData* MeshData = Section->bNextFrameMeshDataSelected ? Section->NextFrameMeshData : Section->MeshData;
-						for (const FGeometryCacheMeshBatchInfo& BatchInfo : MeshData->BatchesInfo)
-						{
-							FRayTracingGeometrySegment Segment;
-							Segment.FirstPrimitive = BatchInfo.StartIndex / 3;
-							Segment.NumPrimitives = BatchInfo.NumTriangles;
-							Segment.VertexBuffer = Section->PositionBuffers[PositionBufferIndex].VertexBufferRHI;
-							Segment.MaxVertices = Section->PositionBuffers[PositionBufferIndex].GetSizeInBytes() / Segment.VertexBufferStride; // conservative estimate
-							Segments.Add(Segment);
-							Initializer.TotalPrimitiveCount += BatchInfo.NumTriangles;
-						}
-
-						Initializer.Segments = Segments;
-
-						// The geometry is not considered valid for initialization unless it has any triangles
-						if (Initializer.TotalPrimitiveCount > 0)
-						{
-							Initializer.IndexBuffer = Section->IndexBuffer.IndexBufferRHI;
-						}
-
-						Section->RayTracingGeometry.SetInitializer(Initializer);
-						Section->RayTracingGeometry.InitResource(RHICmdList);
-					}
-				}
+				SceneProxy->InitRayTracing(RHICmdList);
 			});
 		}
 #endif
@@ -632,6 +591,11 @@ void FGeometryCacheSceneProxy::UpdateAnimation(FRHICommandListBase& RHICmdList,
 				{
 					continue;
 				}
+
+				if (!Section->bInitializedRayTracing)
+				{
+					InitRayTracing(RHICmdList);
+				}
 			
 				const int PositionBufferIndex = Section->CurrentPositionBufferIndex != -1 ? Section->CurrentPositionBufferIndex % 2 : 0;
 				const FGeometryCacheMeshData* MeshData = Section->bNextFrameMeshDataSelected ? Section->NextFrameMeshData : Section->MeshData;
@@ -1388,6 +1352,58 @@ void FGeometryCacheSceneProxy::FrameUpdate(FRHICommandListBase& RHICmdList) cons
 	}
 }
 
+#if RHI_RAYTRACING
+void FGeometryCacheSceneProxy::InitRayTracing(FRHICommandListBase& RHICmdList)
+{
+	for (FGeomCacheTrackProxy* Section : Tracks)
+	{
+		if (Section != nullptr)
+		{
+			const FVisibilitySample& VisibilitySample = Section->GetVisibilitySample(Time, bLooping);
+			if (!VisibilitySample.bVisibilityState)
+			{
+				continue;
+			}
+
+			FRayTracingGeometryInitializer Initializer;
+			Initializer.DebugName = RayTracingDebugName;
+			const int PositionBufferIndex = Section->CurrentPositionBufferIndex != -1 ? Section->CurrentPositionBufferIndex % 2 : 0;
+			Initializer.TotalPrimitiveCount = 0;
+			Initializer.GeometryType = RTGT_Triangles;
+			Initializer.bFastBuild = false;
+			Initializer.bAllowCompaction = false;
+
+			TArray<FRayTracingGeometrySegment> Segments;
+			const FGeometryCacheMeshData* MeshData = Section->bNextFrameMeshDataSelected ? Section->NextFrameMeshData : Section->MeshData;
+			for (const FGeometryCacheMeshBatchInfo& BatchInfo : MeshData->BatchesInfo)
+			{
+				FRayTracingGeometrySegment Segment;
+				Segment.FirstPrimitive = BatchInfo.StartIndex / 3;
+				Segment.NumPrimitives = BatchInfo.NumTriangles;
+				Segment.VertexBuffer = Section->PositionBuffers[PositionBufferIndex].VertexBufferRHI;
+				Segment.MaxVertices = Section->PositionBuffers[PositionBufferIndex].GetSizeInBytes() / Segment.VertexBufferStride; // conservative estimate
+				Segments.Add(Segment);
+				Initializer.TotalPrimitiveCount += BatchInfo.NumTriangles;
+			}
+
+			Initializer.Segments = Segments;
+
+			// The geometry is not considered valid for initialization unless it has any triangles
+			if (Initializer.TotalPrimitiveCount > 0)
+			{
+				Initializer.IndexBuffer = Section->IndexBuffer.IndexBufferRHI;
+			}
+
+			Section->RayTracingGeometry.SetInitializer(Initializer);
+			Section->RayTracingGeometry.InitResource(RHICmdList);
+			
+			Section->bInitializedRayTracing = true;
+		}
+	}
+
+}
+#endif
+
 void FGeometryCacheSceneProxy::UpdateSectionWorldMatrix(const int32 SectionIndex, const FMatrix& WorldMatrix)
 {
 	check(SectionIndex < Tracks.Num() && "Section Index out of range");
diff --git a/Engine/Plugins/Runtime/GeometryCache/Source/GeometryCache/Public/GeometryCacheSceneProxy.h b/Engine/Plugins/Runtime/GeometryCache/Source/GeometryCache/Public/GeometryCacheSceneProxy.h
--- a/Engine/Plugins/Runtime/GeometryCache/Source/GeometryCache/Public/GeometryCacheSceneProxy.h
+++ b/Engine/Plugins/Runtime/GeometryCache/Source/GeometryCache/Public/GeometryCacheSceneProxy.h
@@ -272,6 +272,7 @@ public:
 	bool bResourcesInitialized;
 
 #if RHI_RAYTRACING
+	bool bInitializedRayTracing = false;
 	FRayTracingGeometry RayTracingGeometry;
 #endif
 };
@@ -323,6 +324,10 @@ public:
 private:
 	void FrameUpdate(FRHICommandListBase& RHICmdList) const;
 
+#if RHI_RAYTRACING
+	void InitRayTracing(FRHICommandListBase& RHICmdList);
+#endif
+
 	void CreateMeshBatch(
 		FRHICommandListBase& RHICmdList,
 		const FGeomCacheTrackProxy* TrackProxy,
@@ -407,6 +412,9 @@ private:
 
 	/** Function used to create a new track proxy at construction */
 	TFunction<FGeomCacheTrackProxy*()> CreateTrackProxy;
+#if RHI_RAYTRACING
+	FName RayTracingDebugName = NAME_None;
+#endif
 };
 
 #if !defined(GEOMETRY_CACHE_SCENE_PROXY_ISPC_ENABLED_DEFAULT)
