[GeomCache] Some improvements to the memory management. Preventing reallocs, resulting in faster performance.

diff --git a/Engine/Plugins/Runtime/GeometryCache/Source/GeometryCache/Classes/GeometryCacheMeshData.h b/Engine/Plugins/Runtime/GeometryCache/Source/GeometryCache/Classes/GeometryCacheMeshData.h
--- a/Engine/Plugins/Runtime/GeometryCache/Source/GeometryCache/Classes/GeometryCacheMeshData.h
+++ b/Engine/Plugins/Runtime/GeometryCache/Source/GeometryCache/Classes/GeometryCacheMeshData.h
@@ -102,20 +102,6 @@ struct FGeometryCacheMeshData
 	: BoundingBox(EForceInit::ForceInitToZero)
 	, Hash(0)
 	{}
-
-	~FGeometryCacheMeshData()
-	{
-		Positions.Empty();
-		TextureCoordinates.Empty();
-		TangentsX.Empty();
-		TangentsZ.Empty();
-		Colors.Empty();
-		ImportedVertexNumbers.Empty();
-
-		MotionVectors.Empty();
-		BatchesInfo.Empty();
-		BoundingBox.Init();
-	}
 	
 	/** Draw-able vertex data */
 	TArray<FVector3f> Positions;
@@ -137,6 +123,9 @@ struct FGeometryCacheMeshData
 	TArray<uint32> Indices;	
 	/** Info on the vertex attributes */
 	FGeometryCacheVertexInfo VertexInfo;
+
+	/** A reusable array of bytes used to decode frame data. */
+	TArray<uint8> TempBytes;
 		
 	/** Serialization for FVertexAnimationSample. */
 	friend GEOMETRYCACHE_API FArchive& operator<<(FArchive& Ar, FGeometryCacheMeshData& Mesh);
diff --git a/Engine/Plugins/Runtime/GeometryCache/Source/GeometryCache/Private/CodecV1.cpp b/Engine/Plugins/Runtime/GeometryCache/Source/GeometryCache/Private/CodecV1.cpp
--- a/Engine/Plugins/Runtime/GeometryCache/Source/GeometryCache/Private/CodecV1.cpp
+++ b/Engine/Plugins/Runtime/GeometryCache/Source/GeometryCache/Private/CodecV1.cpp
@@ -1114,9 +1114,9 @@ bool FCodecV1Decoder::DecodeFrameData(FBufferReader& Reader, FGeometryCacheMeshD
 	const uint32 NumVertices = Header.VertexCount;
 
 	uint32 NumChunks;
-	TArray<uint32> ChunkOffsets;
+	TArray<uint32, TInlineAllocator<8>> ChunkOffsets;
 
-	bool IsChunked = false;
+	bool IsChunked;
 	if (Header.Magic == V1_MAGIC)
 	{
 		// Compatibility with old files. Just treat everything as one big chunk
@@ -1135,8 +1135,8 @@ bool FCodecV1Decoder::DecodeFrameData(FBufferReader& Reader, FGeometryCacheMeshD
 	DecodingContext.MeshData = &OutMeshData;	
 	
 	// Read the payload in memory to pass to the bit reader
-	TArray<uint8> Bytes;
-	Bytes.AddUninitialized(Header.PayloadSize + 16);	// Overallocate by 16 bytes to ensure BitReader can safely perform uint64 reads.
+	TArray<uint8>& Bytes = OutMeshData.TempBytes;
+	Bytes.SetNumUninitialized(Header.PayloadSize + 16, EAllowShrinking::No); // Overallocate by 16 bytes to ensure BitReader can safely perform uint64 reads.
 	Reader.Serialize(Bytes.GetData(), Header.PayloadSize);
 	FHuffmanBitStreamReader BitReader(Bytes.GetData(), Bytes.Num());
 	
@@ -1153,68 +1153,60 @@ bool FCodecV1Decoder::DecodeFrameData(FBufferReader& Reader, FGeometryCacheMeshD
 	{
 		if (!OutMeshData.VertexInfo.bConstantIndices)
 		{
-			OutMeshData.Indices.Empty(Header.IndexCount);
-			OutMeshData.Indices.AddUninitialized(Header.IndexCount);
+			OutMeshData.Indices.SetNumUninitialized(Header.IndexCount, EAllowShrinking::No);
 		}
 
-		OutMeshData.Positions.Empty(Header.VertexCount);
-		OutMeshData.Positions.AddUninitialized(Header.VertexCount);
+		OutMeshData.Positions.SetNumUninitialized(Header.VertexCount, EAllowShrinking::No);
 
-		OutMeshData.ImportedVertexNumbers.Empty(Header.VertexCount);
+		OutMeshData.ImportedVertexNumbers.Reset(Header.VertexCount);
 		if (OutMeshData.VertexInfo.bHasImportedVertexNumbers)
 		{
-			OutMeshData.ImportedVertexNumbers.AddUninitialized(Header.VertexCount);
+			OutMeshData.ImportedVertexNumbers.SetNumUninitialized(Header.VertexCount, EAllowShrinking::No);
 		}
 
-		OutMeshData.Colors.Empty(Header.VertexCount);
 		if (OutMeshData.VertexInfo.bHasColor0)
 		{
-			OutMeshData.Colors.AddUninitialized(Header.VertexCount);
+			OutMeshData.Colors.SetNumUninitialized(Header.VertexCount, EAllowShrinking::No);
 		}
 		else
 		{
-			OutMeshData.Colors.AddZeroed(Header.VertexCount);
+			OutMeshData.Colors.SetNumZeroed(Header.VertexCount, EAllowShrinking::No);
 		}	
 		
-		OutMeshData.TangentsX.Empty(Header.VertexCount);
 		if (OutMeshData.VertexInfo.bHasTangentX)
 		{
-			OutMeshData.TangentsX.AddUninitialized(Header.VertexCount);
+			OutMeshData.TangentsX.SetNumUninitialized(Header.VertexCount, EAllowShrinking::No);
 		}
 		else
 		{
-			OutMeshData.TangentsX.AddZeroed(Header.VertexCount);
+			OutMeshData.TangentsX.SetNumZeroed(Header.VertexCount, EAllowShrinking::No);
 		}
 		
-		OutMeshData.TangentsZ.Empty(Header.VertexCount);
 		if (OutMeshData.VertexInfo.bHasTangentZ)
 		{
-			OutMeshData.TangentsZ.AddUninitialized(Header.VertexCount);
+			OutMeshData.TangentsZ.SetNumUninitialized(Header.VertexCount, EAllowShrinking::No);
 		}
 		else
 		{
-			OutMeshData.TangentsZ.AddZeroed(Header.VertexCount);
+			OutMeshData.TangentsZ.SetNumZeroed(Header.VertexCount, EAllowShrinking::No);
 		}
 		
-		OutMeshData.TextureCoordinates.Empty(Header.VertexCount);
 		if (OutMeshData.VertexInfo.bHasUV0)
 		{
-			OutMeshData.TextureCoordinates.AddUninitialized(Header.VertexCount);
+			OutMeshData.TextureCoordinates.SetNumUninitialized(Header.VertexCount, EAllowShrinking::No);
 		}
 		else
 		{
-			OutMeshData.TextureCoordinates.AddZeroed(Header.VertexCount);
-		}
-		
+			OutMeshData.TextureCoordinates.SetNumZeroed(Header.VertexCount, EAllowShrinking::No);
+		}	
 		
-		OutMeshData.MotionVectors.Empty(Header.VertexCount);
 		if (OutMeshData.VertexInfo.bHasMotionVectors)
 		{
-			OutMeshData.MotionVectors.AddUninitialized(Header.VertexCount);
+			OutMeshData.MotionVectors.SetNumUninitialized(Header.VertexCount, EAllowShrinking::No);
 		}
 		else
 		{
-			OutMeshData.MotionVectors.AddZeroed(Header.VertexCount);
+			OutMeshData.MotionVectors.SetNumZeroed(Header.VertexCount, EAllowShrinking::No);
 		}
 	}
 	
diff --git a/Engine/Plugins/Runtime/GeometryCache/Source/GeometryCache/Private/GeometryCacheTrackStreamable.cpp b/Engine/Plugins/Runtime/GeometryCache/Source/GeometryCache/Private/GeometryCacheTrackStreamable.cpp
--- a/Engine/Plugins/Runtime/GeometryCache/Source/GeometryCache/Private/GeometryCacheTrackStreamable.cpp
+++ b/Engine/Plugins/Runtime/GeometryCache/Source/GeometryCache/Private/GeometryCacheTrackStreamable.cpp
@@ -635,7 +635,6 @@ void FGeometryCacheTrackStreamableRenderResource::ReleaseRHI()
 bool FGeometryCacheTrackStreamableRenderResource::UpdateMeshData(float Time, bool bLooping, int32& InOutMeshSampleIndex, FGeometryCacheMeshData& OutMeshData)
 {
 	int32 InMeshIndex = InOutMeshSampleIndex;
-	float TimeNonLooping = (bLooping) ? GeometyCacheHelpers::WrapAnimationTime(Time, Track->Duration) : Time;
 	int32 SampleIndexToDecode = Track->FindSampleIndexFromTime(Time, bLooping);
 	if (SampleIndexToDecode == InOutMeshSampleIndex)
 	{
