GeometryCache: Fixed potentiel crash when a material on a geometry cache is somehow null.

#rnx

diff --git a/Engine/Plugins/Runtime/GeometryCache/Source/GeometryCache/Private/GeometryCacheSceneProxy.cpp b/Engine/Plugins/Runtime/GeometryCache/Source/GeometryCache/Private/GeometryCacheSceneProxy.cpp
--- a/Engine/Plugins/Runtime/GeometryCache/Source/GeometryCache/Private/GeometryCacheSceneProxy.cpp
+++ b/Engine/Plugins/Runtime/GeometryCache/Source/GeometryCache/Private/GeometryCacheSceneProxy.cpp
@@ -458,7 +458,9 @@ void FGeometryCacheSceneProxy::GetDynamicMeshElements(const TArray<const FSceneV
 						// Apply view mode material overrides
 						const int32 MaterialIndex = TrackProxy->Materials.IsValidIndex(BatchInfo.MaterialIndex) ? BatchInfo.MaterialIndex :
 													TrackProxy->Materials.IsValidIndex(BatchIndex) ? BatchIndex : 0; // extra precaution in case of bad data
-						FMaterialRenderProxy* MaterialProxy = bWireframe ? WireframeMaterialInstance : TrackProxy->Materials[MaterialIndex]->GetRenderProxy();
+						FMaterialRenderProxy* MaterialProxy = bWireframe ? WireframeMaterialInstance :
+															  TrackProxy->Materials[MaterialIndex] != nullptr ?
+															  TrackProxy->Materials[MaterialIndex]->GetRenderProxy() : nullptr;
 						MeshBatch.bWireframe = bWireframe;
 						MeshBatch.MaterialRenderProxy = MaterialProxy;
 
