Fix for potential dangling vertex stream pointers when setting up dummy data object for PSO vertex declaration setup

diff --git a/Engine/Plugins/Runtime/GeometryCache/Source/GeometryCache/Private/GeometryCacheSceneProxy.cpp b/Engine/Plugins/Runtime/GeometryCache/Source/GeometryCache/Private/GeometryCacheSceneProxy.cpp
--- a/Engine/Plugins/Runtime/GeometryCache/Source/GeometryCache/Private/GeometryCacheSceneProxy.cpp
+++ b/Engine/Plugins/Runtime/GeometryCache/Source/GeometryCache/Private/GeometryCacheSceneProxy.cpp
@@ -1508,6 +1508,17 @@ void FGeomCacheVertexFactory::InitDummyComponentVF(FGeometryCacheVertexVertexFac
 	FGeomCacheVertexBuffer TextureCoordinatesBuffer;
 	FGeomCacheColorBuffer ColorBuffer;
 	InitComponentVF(&PositionBuffers[0], &PositionBuffers[1], &TangentXBuffer, &TangentZBuffer, &TextureCoordinatesBuffer, &ColorBuffer, OutData);
+
+	// Clear out the pointer entries in OutData again
+	OutData.PositionComponent.VertexBuffer = nullptr;
+	for (FVertexStreamComponent& TextureStreamComponent : OutData.TextureCoordinates)
+	{
+		TextureStreamComponent.VertexBuffer = nullptr;
+	}
+	OutData.TangentBasisComponents[0].VertexBuffer = nullptr;
+	OutData.TangentBasisComponents[1].VertexBuffer = nullptr;
+	OutData.ColorComponent.VertexBuffer = nullptr;
+	OutData.MotionBlurDataComponent.VertexBuffer = nullptr;
 }
 
 void FGeomCacheVertexFactory::InitComponentVF(const FVertexBuffer* PositionBuffer, const FVertexBuffer* MotionBlurDataBuffer, const FVertexBuffer* TangentXBuffer, const FVertexBuffer* TangentZBuffer, const FVertexBuffer* TextureCoordinateBuffer, const FVertexBuffer* ColorBuffer, FGeometryCacheVertexVertexFactory::FDataType& OutData)
@@ -1515,6 +1526,7 @@ void FGeomCacheVertexFactory::InitComponentVF(const FVertexBuffer* PositionBuffe
 	// Initialize the vertex factory's stream components.
 	OutData.PositionComponent = FVertexStreamComponent(PositionBuffer, 0, sizeof(FVector3f), VET_Float3);
 
+	OutData.TextureCoordinates.Empty();
 	OutData.TextureCoordinates.Add(FVertexStreamComponent(TextureCoordinateBuffer, 0, sizeof(FVector2f), VET_Float2));
 	OutData.TangentBasisComponents[0] = FVertexStreamComponent(TangentXBuffer, 0, sizeof(FPackedNormal), VET_PackedNormal);
 	OutData.TangentBasisComponents[1] = FVertexStreamComponent(TangentZBuffer, 0, sizeof(FPackedNormal), VET_PackedNormal);
