GeometryCache: Removed the call to FlushRenderingCommand introduced in cl 23782427 since it's already called when CacheData is destroyed.

diff --git a/Engine/Plugins/Runtime/GeometryCache/Source/GeometryCache/Private/GeometryCacheStreamingManager.cpp b/Engine/Plugins/Runtime/GeometryCache/Source/GeometryCache/Private/GeometryCacheStreamingManager.cpp
--- a/Engine/Plugins/Runtime/GeometryCache/Source/GeometryCache/Private/GeometryCacheStreamingManager.cpp
+++ b/Engine/Plugins/Runtime/GeometryCache/Source/GeometryCache/Private/GeometryCacheStreamingManager.cpp
@@ -228,22 +228,7 @@ void FGeometryCacheStreamingManager::RemoveGeometryCache(UGeometryCacheTrackStre
 {
 	check(IsInGameThread());
 
-	// Reorganized to not hold the lock while calling FlushRenderingCommands();
-	// This causes a deadlock with MapChunk.
-
-	bool bCacheFound = false;
-	{
-		FRWScopeLock ReadLock(StreamingGeometryCachesLock, SLT_ReadOnly);
-		bCacheFound = StreamingGeometryCaches.Contains(Cache);
-	}
-
-	// FlushRenderingCommands first, without holding the lock.
-	if (bCacheFound)
-	{
-		// Flush the render thread so any decoding still happening is finished and thus no maps held by the render thread.
-		FlushRenderingCommands();
-	}
-
+	// Reorganized to not hold the lock while calling FlushRenderingCommand()
 	FStreamingGeometryCacheData* CacheData = nullptr;
 	{
 		FRWScopeLock WriteLock(StreamingGeometryCachesLock, SLT_Write);
@@ -258,7 +243,7 @@ void FGeometryCacheStreamingManager::RemoveGeometryCache(UGeometryCacheTrackStre
 	if (CacheData != nullptr)
 	{
 		check(CacheData != nullptr);
-		delete CacheData;	// This will FlushRenderingCommands again so we move it outside the lock.
+		delete CacheData;	// This will FlushRenderingCommand outside the lock.
 	}
 }
 
