* Fix geometry cache recreating SRV with incorrect type when the geometry cache data size change.
* Fix incorrect parameter order when creating hair mesh data.

diff --git a/Engine/Plugins/Runtime/GeometryCache/Source/GeometryCache/Private/GeometryCacheSceneProxy.cpp b/Engine/Plugins/Runtime/GeometryCache/Source/GeometryCache/Private/GeometryCacheSceneProxy.cpp
--- a/Engine/Plugins/Runtime/GeometryCache/Source/GeometryCache/Private/GeometryCacheSceneProxy.cpp
+++ b/Engine/Plugins/Runtime/GeometryCache/Source/GeometryCache/Private/GeometryCacheSceneProxy.cpp
@@ -1710,6 +1710,9 @@ void FGeomCacheVertexBuffer::UpdateRaw(FRHICommandListBase& RHICmdList, const vo
 	int32 NewSizeInBytes = ItemSizeBytes * NumItems;
 	bool bCanMemcopy = ItemSizeBytes == ItemStrideBytes;
 
+	// Preserve SRV format if the SRV needs to be recreated
+	EPixelFormat Format = BufferSRV ? BufferSRV->GetDesc().Buffer.SRV.Format : PF_R32_FLOAT;
+
 	bool bReallocate = false;
 	if (NewSizeInBytes > SizeInBytes)
 	{
@@ -1744,7 +1747,7 @@ void FGeomCacheVertexBuffer::UpdateRaw(FRHICommandListBase& RHICmdList, const vo
 			VertexBufferRHI, 
 			FRHIViewDesc::CreateBufferSRV()
 				.SetType(FRHIViewDesc::EBufferType::Typed)
-				.SetFormat(PF_R32_FLOAT));
+				.SetFormat(Format));
 	}
 }
 
diff --git a/Engine/Plugins/Runtime/HairStrands/Source/HairStrandsCore/Private/GroomManager.cpp b/Engine/Plugins/Runtime/HairStrands/Source/HairStrandsCore/Private/GroomManager.cpp
--- a/Engine/Plugins/Runtime/HairStrands/Source/HairStrandsCore/Private/GroomManager.cpp
+++ b/Engine/Plugins/Runtime/HairStrands/Source/HairStrandsCore/Private/GroomManager.cpp
@@ -310,7 +310,7 @@ struct FHairGeometryCache
 				const EPixelFormat TangentFormat = TangentSRV->GetDesc().Buffer.SRV.Format;
 				const uint32 TangentElementCount = InLODData->StaticVertexBuffers.StaticMeshVertexBuffer.GetNumVertices();
 				check(TangentSRV->GetDesc().Buffer.SRV.BufferType == FRHIViewDesc::EBufferType::Typed);
-				Data.TangentBuffer = GraphBuilder.CreateBuffer(FRDGBufferDesc::CreateBufferDesc(TangentElementCount, TangentStride), TEXT("Hair.SkinnedDeformedTangents"));;
+				Data.TangentBuffer = GraphBuilder.CreateBuffer(FRDGBufferDesc::CreateBufferDesc(TangentStride, TangentElementCount), TEXT("Hair.SkinnedDeformedTangents"));
 				Data.TangentSRV = GraphBuilder.CreateSRV(Data.TangentBuffer, TangentFormat);
 			}
 
