GeometryCache: Fixed crash in FGeometryCacheStreamingManager where it could reference stale GeometryCache components.

#rnx

diff --git a/Engine/Plugins/Runtime/GeometryCache/Source/GeometryCache/Private/GeometryCacheComponent.cpp b/Engine/Plugins/Runtime/GeometryCache/Source/GeometryCache/Private/GeometryCacheComponent.cpp
--- a/Engine/Plugins/Runtime/GeometryCache/Source/GeometryCache/Private/GeometryCacheComponent.cpp
+++ b/Engine/Plugins/Runtime/GeometryCache/Source/GeometryCache/Private/GeometryCacheComponent.cpp
@@ -39,6 +39,10 @@ void UGeometryCacheComponent::BeginDestroy()
 {
 	Super::BeginDestroy();
 	ReleaseResources();
+	// Make sure to remove the component from the streaming manager
+	// There are some flows with Sequencer and PIE where it is removed in OnUnregister but
+	// re-added through SetGeometryCache
+	IGeometryCacheStreamingManager::Get().RemoveStreamingComponent(this);
 }
 
 void UGeometryCacheComponent::FinishDestroy()
