Call static FTSTicker::RemoveTicker in ~FGeometryCacheStreamer to fix nodiscard warning on Clang 20

diff --git a/Engine/Plugins/Runtime/GeometryCache/Source/GeometryCacheStreamer/Private/GeometryCacheStreamer.cpp b/Engine/Plugins/Runtime/GeometryCache/Source/GeometryCacheStreamer/Private/GeometryCacheStreamer.cpp
--- a/Engine/Plugins/Runtime/GeometryCache/Source/GeometryCacheStreamer/Private/GeometryCacheStreamer.cpp
+++ b/Engine/Plugins/Runtime/GeometryCache/Source/GeometryCacheStreamer/Private/GeometryCacheStreamer.cpp
@@ -92,7 +92,7 @@ FGeometryCacheStreamer::~FGeometryCacheStreamer()
 {
 	if (TickHandle.IsValid())
 	{
-		FTSTicker::GetCoreTicker().RemoveTicker(TickHandle);
+		FTSTicker::RemoveTicker(TickHandle);
 		TickHandle.Reset();
 	}
 }
