[HWRT] Fix FGeometryCacheSceneProxy modifying RayTracingGeometry.Initializer before calling ReleaseRHI().

diff --git a/Engine/Plugins/Runtime/GeometryCache/Source/GeometryCache/Private/GeometryCacheSceneProxy.cpp b/Engine/Plugins/Runtime/GeometryCache/Source/GeometryCache/Private/GeometryCacheSceneProxy.cpp
--- a/Engine/Plugins/Runtime/GeometryCache/Source/GeometryCache/Private/GeometryCacheSceneProxy.cpp
+++ b/Engine/Plugins/Runtime/GeometryCache/Source/GeometryCache/Private/GeometryCacheSceneProxy.cpp
@@ -613,31 +613,32 @@ void FGeometryCacheSceneProxy::UpdateAnimation(FRHICommandListBase& RHICmdList,
 				const FGeometryCacheMeshData* MeshData = Section->bNextFrameMeshDataSelected ? Section->NextFrameMeshData : Section->MeshData;
 				const uint32 IndexBufferNumTriangles = Section->IndexBuffer.NumValidIndices / 3;
 
-				TArray<FRayTracingGeometrySegment>& Segments = Section->RayTracingGeometry.Initializer.Segments;
-
 				// Check if a full RaytracingGeometry object needs to be recreated. 
 				// Recreate when:
 				// - index buffer changes (grew in size)
 				// - total primitive count changes
 				// - segment count or vertex count changed (change BLAS size)
 				bool bRequireRecreate = false;
-				bRequireRecreate = bRequireRecreate || Segments.Num() != MeshData->BatchesInfo.Num();				
+				bRequireRecreate = bRequireRecreate || Section->RayTracingGeometry.Initializer.IndexBuffer != Section->IndexBuffer.IndexBufferRHI;
+				bRequireRecreate = bRequireRecreate || Section->RayTracingGeometry.Initializer.Segments.Num() != MeshData->BatchesInfo.Num();
 
 				// Validate the max vertex count on all segments 
 				if (!bRequireRecreate)
 				{
-					for (int32 SegmentIndex = 0; SegmentIndex < Segments.Num(); ++SegmentIndex)
+					for (int32 SegmentIndex = 0; SegmentIndex < MeshData->BatchesInfo.Num(); ++SegmentIndex)
 					{
 						const FGeometryCacheMeshBatchInfo& BatchInfo = MeshData->BatchesInfo[SegmentIndex];
-						const FRayTracingGeometrySegment& Segment = Segments[SegmentIndex];
+						const FRayTracingGeometrySegment& Segment = Section->RayTracingGeometry.Initializer.Segments[SegmentIndex];
 
 						int32 MaxSegmentVertices = Section->PositionBuffers[PositionBufferIndex].GetSizeInBytes() / Segment.VertexBufferStride; // conservative estimate
 						bRequireRecreate = bRequireRecreate || Segment.MaxVertices != MaxSegmentVertices;
 					}
 				}
 
-				uint32 TotalPrimitiveCount = 0;				
-				Segments.Reset();
+				uint32 TotalPrimitiveCount = 0;
+				TArray<FRayTracingGeometrySegment> Segments;
+				Segments.Reserve(MeshData->BatchesInfo.Num());
+
 				for (const FGeometryCacheMeshBatchInfo& BatchInfo : MeshData->BatchesInfo)
 				{
 					FRayTracingGeometrySegment Segment;
@@ -660,18 +661,22 @@ void FGeometryCacheSceneProxy::UpdateAnimation(FRHICommandListBase& RHICmdList,
 					TotalPrimitiveCount += Segment.NumPrimitives;
 				}
 
-				bRequireRecreate = bRequireRecreate || Section->RayTracingGeometry.Initializer.IndexBuffer != Section->IndexBuffer.IndexBufferRHI;
 				bRequireRecreate = bRequireRecreate || Section->RayTracingGeometry.Initializer.TotalPrimitiveCount != TotalPrimitiveCount;
 
-				Section->RayTracingGeometry.Initializer.IndexBuffer = Section->IndexBuffer.IndexBufferRHI;
-				Section->RayTracingGeometry.Initializer.TotalPrimitiveCount = TotalPrimitiveCount;
-
 				if (bRequireRecreate)
 				{
-					Section->RayTracingGeometry.UpdateRHI(RHICmdList);
+					Section->RayTracingGeometry.ReleaseRHI();
+
+					Section->RayTracingGeometry.Initializer.IndexBuffer = Section->IndexBuffer.IndexBufferRHI;
+					Section->RayTracingGeometry.Initializer.TotalPrimitiveCount = TotalPrimitiveCount;
+					Section->RayTracingGeometry.Initializer.Segments = MoveTemp(Segments);
+
+					Section->RayTracingGeometry.InitRHI(RHICmdList);
 				}
 				else if(Section->RayTracingGeometry.IsValid() && !Section->RayTracingGeometry.IsEvicted())
 				{
+					Section->RayTracingGeometry.Initializer.Segments = MoveTemp(Segments);
+
 					// Request full build on same geometry because data might have changed to much for update call?
 					FRayTracingGeometryBuildParams BuildParams;
 					BuildParams.Geometry = Section->RayTracingGeometry.GetRHI();
