Fix hang when previewing Geometry Cache with 1 frame.

diff --git a/Engine/Plugins/Runtime/GeometryCache/Source/GeometryCacheEd/Private/GeometryCacheTimelineBindingAsset.cpp b/Engine/Plugins/Runtime/GeometryCache/Source/GeometryCacheEd/Private/GeometryCacheTimelineBindingAsset.cpp
--- a/Engine/Plugins/Runtime/GeometryCache/Source/GeometryCacheEd/Private/GeometryCacheTimelineBindingAsset.cpp
+++ b/Engine/Plugins/Runtime/GeometryCache/Source/GeometryCacheEd/Private/GeometryCacheTimelineBindingAsset.cpp
@@ -70,13 +70,16 @@ FFrameNumber FGeometryCacheTimelineBindingAsset::GetScrubPosition() const
 		else
 		{
 			SampleTime = PreviewComponent->GetElapsedTime();
-			while (SampleTime > PreviewComponent->GetDuration())
+			if (PreviewComponent->GetDuration() > UE_SMALL_NUMBER)
 			{
-				SampleTime -= PreviewComponent->GetDuration();
-			}
-			while (SampleTime < 0)
-			{
-				SampleTime += PreviewComponent->GetDuration();
+				while (SampleTime > PreviewComponent->GetDuration())
+				{
+					SampleTime -= PreviewComponent->GetDuration();
+				}
+				while (SampleTime < 0)
+				{
+					SampleTime += PreviewComponent->GetDuration();
+				}
 			}
 			SampleTime = FMath::Clamp(SampleTime, 0.0f, PreviewComponent->GetDuration());
 		}
