GeometryCache: Fix material indexing going out of bounds in edge cases.

diff --git a/Engine/Plugins/Runtime/GeometryCache/Source/GeometryCache/Private/GeometryCacheSceneProxy.cpp b/Engine/Plugins/Runtime/GeometryCache/Source/GeometryCache/Private/GeometryCacheSceneProxy.cpp
--- a/Engine/Plugins/Runtime/GeometryCache/Source/GeometryCache/Private/GeometryCacheSceneProxy.cpp
+++ b/Engine/Plugins/Runtime/GeometryCache/Source/GeometryCache/Private/GeometryCacheSceneProxy.cpp
@@ -456,10 +456,9 @@ void FGeometryCacheSceneProxy::GetDynamicMeshElements(const TArray<const FSceneV
 #endif
 
 						// Apply view mode material overrides
-						const int32 MaterialIndex = TrackProxy->Materials.IsValidIndex(BatchInfo.MaterialIndex) ? BatchInfo.MaterialIndex :
-													TrackProxy->Materials.IsValidIndex(BatchIndex) ? BatchIndex : 0; // extra precaution in case of bad data
+						const int32 MaterialIndex = TrackProxy->Materials.IsValidIndex(BatchInfo.MaterialIndex) ? BatchInfo.MaterialIndex : BatchIndex;
 						FMaterialRenderProxy* MaterialProxy = bWireframe ? WireframeMaterialInstance :
-															  TrackProxy->Materials[MaterialIndex] != nullptr ?
+															  TrackProxy->Materials.IsValidIndex(MaterialIndex) && TrackProxy->Materials[MaterialIndex] != nullptr ? // extra precaution in case of bad data
 															  TrackProxy->Materials[MaterialIndex]->GetRenderProxy() : nullptr;
 						MeshBatch.bWireframe = bWireframe;
 						MeshBatch.MaterialRenderProxy = MaterialProxy;
