GeometryCacheComponent: Added a new "Interpolate Frames" property to control subframe interpolation (for constant topologies) on a per-instance basis. This property can also be animated in a level sequence.
The GeometryCache.InterpolateFrames cvar still exists and controls the feature globally and supersedes the property value when set to false.

#Synced-CL 50979049

diff --git a/Engine/Plugins/Runtime/GeometryCache/Source/GeometryCache/Classes/GeometryCacheComponent.h b/Engine/Plugins/Runtime/GeometryCache/Source/GeometryCache/Classes/GeometryCacheComponent.h
--- a/Engine/Plugins/Runtime/GeometryCache/Source/GeometryCache/Classes/GeometryCacheComponent.h
+++ b/Engine/Plugins/Runtime/GeometryCache/Source/GeometryCache/Classes/GeometryCacheComponent.h
@@ -125,6 +125,14 @@ public:
 	UFUNCTION(BlueprintCallable, Category = "Components|GeometryCache")
 	UE_API void SetLooping( const bool bNewLooping);
 
+	/** Get whether this GeometryCache is interpolating frames. */
+	UFUNCTION(BlueprintCallable, Category = "Components|GeometryCache")
+	UE_API bool IsInterpolatingFrames() const;
+
+	/** Set whether this GeometryCache is interpolating frames. */
+	UFUNCTION(BlueprintCallable, Category = "Components|GeometryCache")
+	UE_API void SetInterpolateFrames(const bool bNewInterpolating);
+
 	/** Get whether this GeometryCache is extrapolating frames. */
 	UFUNCTION(BlueprintCallable, Category = "Components|GeometryCache")
 	UE_API bool IsExtrapolatingFrames() const;
@@ -299,6 +307,10 @@ protected:
 	UPROPERTY(EditAnywhere, Interp, Category = GeometryCache)
 	bool bLooping;
 
+	/** Enable frame interpolation for sub-frame sampling of constant topologies */
+	UPROPERTY(EditAnywhere, Category = GeometryCache, AdvancedDisplay)
+	bool bInterpolateFrames;
+
 	/** Enable frame extrapolation for sub-frame sampling of non-constant topologies with imported motion vectors */
 	UPROPERTY(EditAnywhere, Category = GeometryCache, AdvancedDisplay)
 	bool bExtrapolateFrames;
diff --git a/Engine/Plugins/Runtime/GeometryCache/Source/GeometryCache/Private/GeometryCacheComponent.cpp b/Engine/Plugins/Runtime/GeometryCache/Source/GeometryCache/Private/GeometryCacheComponent.cpp
--- a/Engine/Plugins/Runtime/GeometryCache/Source/GeometryCache/Private/GeometryCacheComponent.cpp
+++ b/Engine/Plugins/Runtime/GeometryCache/Source/GeometryCache/Private/GeometryCacheComponent.cpp
@@ -26,6 +26,7 @@ UGeometryCacheComponent::UGeometryCacheComponent(const FObjectInitializer& Objec
 	ElapsedTime = 0.0f;
 	bRunning = true;
 	bLooping = true;
+	bInterpolateFrames = true;
 	bExtrapolateFrames = true;
 	PlayDirection = 1.0f;
 	StartTimeOffset = 0.0f;
@@ -106,13 +107,19 @@ void UGeometryCacheComponent::JumpAnimationToTime(float Time, bool bInIsRunning,
 	if (FGeometryCacheSceneProxy* CastedProxy = static_cast<FGeometryCacheSceneProxy*>(SceneProxy))
 	{
 		FGeometryCacheSceneProxy* InSceneProxy = CastedProxy;
-		float AnimationTime = Time;
-		float InPlaybackSpeed = ActualPlaybackSpeed;
-		float InMotionVectorScale = GetMotionVectorScale();
+		FGeometryCacheSceneProxy::FAnimationUpdateParameters UpdatedValues =
+		{
+			Time,
+			bInIsLooping,
+			bInBackwards,
+			ActualPlaybackSpeed,
+			GetMotionVectorScale(),
+			IsInterpolatingFrames()
+		};
 		ENQUEUE_RENDER_COMMAND(FGeometryCacheUpdateAnimation)(
-			[InSceneProxy, AnimationTime, bInIsLooping, bInBackwards, InPlaybackSpeed, InMotionVectorScale](FRHICommandList& RHICmdList)
+			[InSceneProxy, UpdatedValues](FRHICommandList& RHICmdList)
 			{
-				InSceneProxy->UpdateAnimation(RHICmdList, AnimationTime, bInIsLooping, bInBackwards, InPlaybackSpeed, InMotionVectorScale);
+				InSceneProxy->UpdateAnimation(RHICmdList, UpdatedValues);
 			});
 	}
 }
@@ -208,15 +215,19 @@ void UGeometryCacheComponent::TickComponent(float DeltaTime, enum ELevelTick Tic
 		if (FGeometryCacheSceneProxy* CastedProxy = static_cast<FGeometryCacheSceneProxy*>(SceneProxy))
 		{
 			FGeometryCacheSceneProxy* InSceneProxy = CastedProxy;
-			float AnimationTime = GetAnimationTime();
-			bool bInLooping = IsLooping();
-			bool bIsPlayingBackwards = IsPlayingReversed();
-			float InPlaybackSpeed = ActualPlaybackSpeed;
-			float InMotionVectorScale = GetMotionVectorScale();
+			FGeometryCacheSceneProxy::FAnimationUpdateParameters UpdatedValues =
+			{
+				GetAnimationTime(),
+				IsLooping(),
+				IsPlayingReversed(),
+				ActualPlaybackSpeed,
+				GetMotionVectorScale(),
+				IsInterpolatingFrames()
+			};
 			ENQUEUE_RENDER_COMMAND(FGeometryCacheUpdateAnimation)(
-				[InSceneProxy, AnimationTime, bInLooping, bIsPlayingBackwards, InPlaybackSpeed, ActualPlaybackSpeed, InMotionVectorScale](FRHICommandList& RHICmdList)
+				[InSceneProxy, UpdatedValues](FRHICommandList& RHICmdList)
 				{
-					InSceneProxy->UpdateAnimation(RHICmdList, AnimationTime, bInLooping, bIsPlayingBackwards, InPlaybackSpeed, InMotionVectorScale);
+					InSceneProxy->UpdateAnimation(RHICmdList, UpdatedValues);
 				});
 		}
 	}
@@ -457,6 +468,16 @@ void UGeometryCacheComponent::SetLooping(const bool bNewLooping)
 	bLooping = bNewLooping;
 }
 
+bool UGeometryCacheComponent::IsInterpolatingFrames() const
+{
+	return bInterpolateFrames;
+}
+
+void UGeometryCacheComponent::SetInterpolateFrames(const bool bNewInterpolating)
+{
+	bInterpolateFrames = bNewInterpolating;
+}
+
 bool UGeometryCacheComponent::IsExtrapolatingFrames() const
 {
 	return bExtrapolateFrames;
diff --git a/Engine/Plugins/Runtime/GeometryCache/Source/GeometryCache/Private/GeometryCacheSceneProxy.cpp b/Engine/Plugins/Runtime/GeometryCache/Source/GeometryCache/Private/GeometryCacheSceneProxy.cpp
--- a/Engine/Plugins/Runtime/GeometryCache/Source/GeometryCache/Private/GeometryCacheSceneProxy.cpp
+++ b/Engine/Plugins/Runtime/GeometryCache/Source/GeometryCache/Private/GeometryCacheSceneProxy.cpp
@@ -100,6 +100,7 @@ FGeometryCacheSceneProxy::FGeometryCacheSceneProxy(UGeometryCacheComponent* Comp
 	Time = Component->GetAnimationTime();
 	bLooping = Component->IsLooping();
 	bExtrapolateFrames = Component->IsExtrapolatingFrames();
+	bInterpolateFrames = Component->IsInterpolatingFrames();
 	bAlwaysHasVelocity = true;
 	PlaybackSpeed = (Component->IsPlaying()) ? Component->GetPlaybackSpeed() : 0.0f;
 	MotionVectorScale = Component->GetMotionVectorScale();
@@ -578,14 +579,29 @@ uint32 FGeometryCacheSceneProxy::GetAllocatedSize(void) const
 	return static_cast<uint32>(FPrimitiveSceneProxy::GetAllocatedSize());
 }
 
-void FGeometryCacheSceneProxy::UpdateAnimation(FRHICommandListBase& RHICmdList, float NewTime, bool bNewLooping, bool bNewIsPlayingBackwards, float NewPlaybackSpeed, float NewMotionVectorScale)
+void FGeometryCacheSceneProxy::UpdateAnimation(FRHICommandListBase& RHICmdList, float NewTime, bool bNewLooping, bool bNewIsPlayingBackwards, float NewPlaybackSpeed, float NewMotionVectorScale, bool bNewInterpolateFrames)
 {
-	Time = NewTime;
-	bLooping = bNewLooping;
-	bIsPlayingBackwards = bNewIsPlayingBackwards;
-	PlaybackSpeed = NewPlaybackSpeed;
-	MotionVectorScale = NewMotionVectorScale;
+	FAnimationUpdateParameters UpdatedValues = 
+		{
+			NewTime,
+			bNewLooping,
+			bNewIsPlayingBackwards,
+			NewPlaybackSpeed,
+			NewMotionVectorScale,
+			bNewInterpolateFrames
+		};
+		UpdateAnimation(RHICmdList, UpdatedValues);
+}
+
+void FGeometryCacheSceneProxy::UpdateAnimation(FRHICommandListBase& RHICmdList, const FAnimationUpdateParameters& UpdatedValues)
+{
+	Time = UpdatedValues.Time;
+	bLooping = UpdatedValues.bLooping;
+	bIsPlayingBackwards = UpdatedValues.bIsPlayingBackwards;
+	PlaybackSpeed = UpdatedValues.PlaybackSpeed;
+	MotionVectorScale = UpdatedValues.MotionVectorScale;
 	UpdatedFrameNum = GFrameNumber + 1;
+	bInterpolateFrames = UpdatedValues.bInterpolateFrames;
 
 	// Always update in render thread regardless of visibility, ray tracing or not
 	FrameUpdate(RHICmdList);
@@ -794,7 +810,7 @@ void FGeometryCacheSceneProxy::FrameUpdate(FRHICommandListBase& RHICmdList) cons
 			}
 
 			// Check if we can interpolate between the two frames we have available
-			const bool bCanInterpolate = TrackProxy->IsTopologyCompatible(TrackProxy->FrameIndex, TrackProxy->NextFrameIndex);
+			const bool bCanInterpolate = TrackProxy->IsTopologyCompatible(TrackProxy->FrameIndex, TrackProxy->NextFrameIndex) && bInterpolateFrames;
 
 			// Check if we have explicit motion vectors
 			const bool bHasMotionVectors = (
diff --git a/Engine/Plugins/Runtime/GeometryCache/Source/GeometryCache/Public/GeometryCacheSceneProxy.h b/Engine/Plugins/Runtime/GeometryCache/Source/GeometryCache/Public/GeometryCacheSceneProxy.h
--- a/Engine/Plugins/Runtime/GeometryCache/Source/GeometryCache/Public/GeometryCacheSceneProxy.h
+++ b/Engine/Plugins/Runtime/GeometryCache/Source/GeometryCache/Public/GeometryCacheSceneProxy.h
@@ -307,7 +307,18 @@ public:
 	UE_API uint32 GetAllocatedSize(void) const;
 	// End FPrimitiveSceneProxy interface.
 
-	UE_API void UpdateAnimation(FRHICommandListBase& RHICmdList, float NewTime, bool bLooping, bool bIsPlayingBackwards, float PlaybackSpeed, float MotionVectorScale);
+	UE_API void UpdateAnimation(FRHICommandListBase& RHICmdList, float NewTime, bool bLooping, bool bIsPlayingBackwards, float PlaybackSpeed, float MotionVectorScale, bool bInterpolateFrames = true);
+
+	struct FAnimationUpdateParameters
+	{
+		float Time;
+		bool bLooping;
+		bool bIsPlayingBackwards;
+		float PlaybackSpeed;
+		float MotionVectorScale;
+		bool bInterpolateFrames;
+	};
+	UE_API void UpdateAnimation(FRHICommandListBase& RHICmdList, const FAnimationUpdateParameters& UpdatedValues);
 
 	/** Update world matrix for specific section */
 	UE_API void UpdateSectionWorldMatrix(const int32 SectionIndex, const FMatrix& WorldMatrix);
@@ -415,6 +426,7 @@ private:
 	uint32 bLooping : 1;
 	uint32 bIsPlayingBackwards : 1;
 	uint32 bExtrapolateFrames : 1;
+	uint32 bInterpolateFrames : 1;
 
 	/** Function used to create a new track proxy at construction */
 	TFunction<FGeomCacheTrackProxy*()> CreateTrackProxy;
