Added options to prevent dirtying the navmesh and collision system with the geometry cache renderer

#Synced-CL 50731356

diff --git a/Engine/Plugins/Runtime/GeometryCache/Source/GeometryCache/Classes/NiagaraGeometryCacheRendererProperties.h b/Engine/Plugins/Runtime/GeometryCache/Source/GeometryCache/Classes/NiagaraGeometryCacheRendererProperties.h
--- a/Engine/Plugins/Runtime/GeometryCache/Source/GeometryCache/Classes/NiagaraGeometryCacheRendererProperties.h
+++ b/Engine/Plugins/Runtime/GeometryCache/Source/GeometryCache/Classes/NiagaraGeometryCacheRendererProperties.h
@@ -148,6 +148,14 @@ public:
 	* Disabling this option is faster, but a particle can get a different component each tick, which can lead to problems with for example motion blur. */
 	UPROPERTY(EditAnywhere, AdvancedDisplay, Category = "GeometryCache")
 	bool bAssignComponentsOnParticleID = true;
+	
+	/** Determines if the geometry component will affect the engine's physics system with its collision shape. */
+	UPROPERTY(EditAnywhere, AdvancedDisplay, Category = "Component")
+	TEnumAsByte<ECollisionEnabled::Type> EnableCollision = ECollisionEnabled::NoCollision;
+	
+	/** If true then the geometry component will invalidate the nav mesh on movement. */
+	UPROPERTY(EditAnywhere, AdvancedDisplay, Category = "Component")
+	bool bAffectsNavigation = false;
 
 	/** If this array has entries, we will create a MaterialInstanceDynamic per Emitter instance from Material and set the Material parameters using the Niagara simulation variables listed.*/
 	UPROPERTY(EditAnywhere, Category = "Bindings")
diff --git a/Engine/Plugins/Runtime/GeometryCache/Source/GeometryCache/Private/NiagaraRendererGeometryCache.cpp b/Engine/Plugins/Runtime/GeometryCache/Source/GeometryCache/Private/NiagaraRendererGeometryCache.cpp
--- a/Engine/Plugins/Runtime/GeometryCache/Source/GeometryCache/Private/NiagaraRendererGeometryCache.cpp
+++ b/Engine/Plugins/Runtime/GeometryCache/Source/GeometryCache/Private/NiagaraRendererGeometryCache.cpp
@@ -170,6 +170,8 @@ UGeometryCacheComponent* FNiagaraRendererGeometryCache::CreateOrGetPooledCompone
 			GeometryComponent->AddTickPrerequisiteComponent(AttachComponent);
 			GeometryComponent->SetLooping(Properties->bIsLooping);
 			GeometryComponent->SetManualTick(true); // we want to tick the component with the delta time of the niagara sim
+			GeometryComponent->SetCanEverAffectNavigation(Properties->bAffectsNavigation);
+			GeometryComponent->SetCollisionEnabled(Properties->EnableCollision);
 
 			if (Emitter->IsLocalSpace())
 			{
