GeometryCache Import: Fixed out of bounds access when checking visibility samples but there's no actual geometry.

#Synced-CL 51095459

diff --git a/Engine/Plugins/Runtime/GeometryCache/Source/GeometryCache/Private/GeometryCacheTrackStreamable.cpp b/Engine/Plugins/Runtime/GeometryCache/Source/GeometryCache/Private/GeometryCacheTrackStreamable.cpp
--- a/Engine/Plugins/Runtime/GeometryCache/Source/GeometryCache/Private/GeometryCacheTrackStreamable.cpp
+++ b/Engine/Plugins/Runtime/GeometryCache/Source/GeometryCache/Private/GeometryCacheTrackStreamable.cpp
@@ -236,7 +236,7 @@ bool UGeometryCacheTrackStreamable::EndCoding()
 		}
 
 		// Make sure the range of all visibility samples covers the entire track duration
-		if (RangeStart < Samples.Last().SampleTime)
+		if (Samples.Num() > 0 && RangeStart < Samples.Last().SampleTime)
 		{
 			TRange<float> VisibilityRange(RangeStart, Samples.Last().SampleTime);
 			FVisibilitySample Sample(bVisible);
